Moved EOS P2P related files from OnlineSubsystemEOS plugin to new SocketSubsystemEOS plugin and refactored them to be used by any subsystem.
Refactored all systems that used them to reflect changes.
Added GetResolvedConnectString to OnlineServicesEOS

diff --git a/Engine/Plugins/Online/EOSShared/Source/EOSShared/Private/EOSShared.cpp b/Engine/Plugins/Online/EOSShared/Source/EOSShared/Private/EOSShared.cpp
--- a/Engine/Plugins/Online/EOSShared/Source/EOSShared/Private/EOSShared.cpp
+++ b/Engine/Plugins/Online/EOSShared/Source/EOSShared/Private/EOSShared.cpp
@@ -1,11 +1,14 @@
 // Copyright Epic Games, Inc. All Rights Reserved.
 
 #include "EOSShared.h"
+#include "EOSSharedTypes.h"
 
 #include "eos_friends_types.h"
 
 DEFINE_LOG_CATEGORY(LogEOSSDK);
 
+bool FCallbackBase::bShouldCancelAllCallbacks = false;
+
 FString LexToString(const EOS_ProductUserId UserId)
 {
 	FString Result;
diff --git a/Engine/Plugins/Online/EOSShared/Source/EOSShared/Public/EOSSharedTypes.h b/Engine/Plugins/Online/EOSShared/Source/EOSShared/Public/EOSSharedTypes.h
new file mode 100644
--- /dev/null
+++ b/Engine/Plugins/Online/EOSShared/Source/EOSShared/Public/EOSSharedTypes.h
@@ -0,0 +1,71 @@
+// Copyright Epic Games, Inc. All Rights Reserved.
+
+#pragma once
+
+#include "CoreGlobals.h"
+
+#if WITH_EOS_SDK
+#if defined(EOS_PLATFORM_BASE_FILE_NAME)
+#include EOS_PLATFORM_BASE_FILE_NAME
+#endif
+
+#include "eos_common.h"
+#endif
+
+// Expect URLs to look like "EOS:PUID:SocketName:Channel"
+#define EOS_CONNECTION_URL_PREFIX TEXT("EOS")
+#define EOS_URL_SEPARATOR TEXT(":")
+
+#ifndef EOS_SOCKETSUBSYSTEM
+#define EOS_SOCKETSUBSYSTEM FName(TEXT("EOS"))
+#endif
+
+/** Used to store a pointer to the EOS callback object without knowing type */
+class EOSSHARED_API FCallbackBase
+{
+	static bool bShouldCancelAllCallbacks;
+
+public:
+	virtual ~FCallbackBase() {}
+	static bool ShouldCancelAllCallbacks() { return FCallbackBase::bShouldCancelAllCallbacks; }
+	static void CancelAllCallbacks() { FCallbackBase::bShouldCancelAllCallbacks = true; }
+};
+
+#if WITH_EOS_SDK
+
+/** Class to handle all callbacks generically using a lambda to process callback results */
+template<typename CallbackFuncType, typename CallbackType>
+class TEOSGlobalCallback :
+	public FCallbackBase
+{
+public:
+	TFunction<void(const CallbackType*)> CallbackLambda;
+
+	TEOSGlobalCallback() = default;
+	virtual ~TEOSGlobalCallback() = default;
+
+
+	CallbackFuncType GetCallbackPtr()
+	{
+		return &CallbackImpl;
+	}
+
+private:
+	static void EOS_CALL CallbackImpl(const CallbackType* Data)
+	{
+		check(IsInGameThread());
+
+		TEOSGlobalCallback* CallbackThis = (TEOSGlobalCallback*)Data->ClientData;
+		check(CallbackThis);
+
+		if (FCallbackBase::ShouldCancelAllCallbacks())
+		{
+			return;
+		}
+
+		check(CallbackThis->CallbackLambda);
+		CallbackThis->CallbackLambda(Data);
+	}
+};
+
+#endif
\ No newline at end of file
diff --git a/Engine/Plugins/Online/Experimental/OnlineServices/Source/OnlineServicesCommon/Private/Online/OnlineServicesCommon.cpp b/Engine/Plugins/Online/Experimental/OnlineServices/Source/OnlineServicesCommon/Private/Online/OnlineServicesCommon.cpp
--- a/Engine/Plugins/Online/Experimental/OnlineServices/Source/OnlineServicesCommon/Private/Online/OnlineServicesCommon.cpp
+++ b/Engine/Plugins/Online/Experimental/OnlineServices/Source/OnlineServicesCommon/Private/Online/OnlineServicesCommon.cpp
@@ -74,6 +74,11 @@ IPrivilegesPtr FOnlineServicesCommon::GetPrivilegesInterface()
 	return IPrivilegesPtr(AsShared(), Get<IPrivileges>());
 }
 
+TOnlineResult<FGetResolvedConnectString> FOnlineServicesCommon::GetResolvedConnectString(FGetResolvedConnectString::Params&& Params)
+{
+	return TOnlineResult<FGetResolvedConnectString>(Errors::NotImplemented());
+}
+
 void FOnlineServicesCommon::RegisterComponents()
 {
 }
diff --git a/Engine/Plugins/Online/Experimental/OnlineServices/Source/OnlineServicesCommon/Public/Online/OnlineServicesCommon.h b/Engine/Plugins/Online/Experimental/OnlineServices/Source/OnlineServicesCommon/Public/Online/OnlineServicesCommon.h
--- a/Engine/Plugins/Online/Experimental/OnlineServices/Source/OnlineServicesCommon/Public/Online/OnlineServicesCommon.h
+++ b/Engine/Plugins/Online/Experimental/OnlineServices/Source/OnlineServicesCommon/Public/Online/OnlineServicesCommon.h
@@ -42,6 +42,7 @@ public:
 	virtual ILobbiesPtr GetLobbiesInterface() override;
 	virtual IConnectivityPtr GetConnectivityInterface() override;
 	virtual IPrivilegesPtr GetPrivilegesInterface() override;
+	virtual TOnlineResult<FGetResolvedConnectString> GetResolvedConnectString(FGetResolvedConnectString::Params&& Params) override;
 
 	// FOnlineServicesCommon
 
diff --git a/Engine/Plugins/Online/Experimental/OnlineServices/Source/OnlineServicesInterface/Public/Online/OnlineServices.h b/Engine/Plugins/Online/Experimental/OnlineServices/Source/OnlineServicesInterface/Public/Online/OnlineServices.h
--- a/Engine/Plugins/Online/Experimental/OnlineServices/Source/OnlineServicesInterface/Public/Online/OnlineServices.h
+++ b/Engine/Plugins/Online/Experimental/OnlineServices/Source/OnlineServicesInterface/Public/Online/OnlineServices.h
@@ -4,6 +4,8 @@
 
 #include "Online/CoreOnline.h"
 #include "Templates/SharedPointer.h"
+#include "Online/OnlineAsyncOpHandle.h"
+#include "Online/OnlineMeta.h"
 
 class FString;
 
@@ -19,6 +21,23 @@ using ILobbiesPtr = TSharedPtr<class ILobbies>;
 using IConnectivityPtr = TSharedPtr<class IConnectivity>;
 using IPrivilegesPtr = TSharedPtr<class IPrivileges>;
 
+struct FGetResolvedConnectString
+{
+	static constexpr TCHAR Name[] = TEXT("GetResolvedConnectString");
+
+	struct Params
+	{
+		FOnlineAccountIdHandle LocalUserId;
+		FOnlineLobbyIdHandle LobbyId;
+		FName PortType;
+	};
+
+	struct Result
+	{
+		FString ResolvedConnectString;
+	};
+};
+
 class ONLINESERVICESINTERFACE_API IOnlineServices
 {
 public:
@@ -67,6 +86,11 @@ public:
 	 *
 	 */
 	virtual IPrivilegesPtr GetPrivilegesInterface() = 0;
+
+	/** 
+	 * Get the connectivity string used for client travel
+	 */
+	virtual TOnlineResult<FGetResolvedConnectString> GetResolvedConnectString(FGetResolvedConnectString::Params&& Params) = 0;
 };
 
 /**
@@ -107,4 +131,19 @@ TSharedPtr<ServicesClass> GetServices(FName InstanceName = NAME_None)
  */
 ONLINESERVICESINTERFACE_API void DestroyServices(EOnlineServices OnlineServices = EOnlineServices::Default, FName InstanceName = NAME_None);
 
+namespace Meta {
+// TODO: Move to OnlineServices_Meta.inl file?
+
+BEGIN_ONLINE_STRUCT_META(FGetResolvedConnectString::Params)
+	ONLINE_STRUCT_FIELD(FGetResolvedConnectString::Params, LocalUserId),
+	ONLINE_STRUCT_FIELD(FGetResolvedConnectString::Params, LobbyId),
+	ONLINE_STRUCT_FIELD(FGetResolvedConnectString::Params, PortType)
+END_ONLINE_STRUCT_META()
+
+BEGIN_ONLINE_STRUCT_META(FGetResolvedConnectString::Result)
+	ONLINE_STRUCT_FIELD(FGetResolvedConnectString::Result, ResolvedConnectString)
+END_ONLINE_STRUCT_META()
+
+/* Meta*/}
+
 /* UE::Online */ }
diff --git a/Engine/Plugins/Online/Experimental/OnlineServicesEOS/OnlineServicesEOS.uplugin b/Engine/Plugins/Online/Experimental/OnlineServicesEOS/OnlineServicesEOS.uplugin
--- a/Engine/Plugins/Online/Experimental/OnlineServicesEOS/OnlineServicesEOS.uplugin
+++ b/Engine/Plugins/Online/Experimental/OnlineServicesEOS/OnlineServicesEOS.uplugin
@@ -36,6 +36,10 @@
 		{
 			"Name": "EOSShared",
 			"Enabled": true
+		},
+		{
+			"Name": "SocketSubsystemEOS",
+			"Enabled": true
 		}
 	]
 }
diff --git a/Engine/Plugins/Online/Experimental/OnlineServicesEOS/Source/OnlineServicesEOS.Build.cs b/Engine/Plugins/Online/Experimental/OnlineServicesEOS/Source/OnlineServicesEOS.Build.cs
--- a/Engine/Plugins/Online/Experimental/OnlineServicesEOS/Source/OnlineServicesEOS.Build.cs
+++ b/Engine/Plugins/Online/Experimental/OnlineServicesEOS/Source/OnlineServicesEOS.Build.cs
@@ -18,8 +18,12 @@ public class OnlineServicesEOS : ModuleRules
 		PrivateDependencyModuleNames.AddRange(
 			new string[] {
 				"CoreOnline",
+				"CoreUObject",
 				"EOSSDK",
 				"EOSShared",
+				"Sockets",
+				"SocketSubsystemEOS",
+				"Engine"
 			}
 		);
 	}
diff --git a/Engine/Plugins/Online/Experimental/OnlineServicesEOS/Source/Private/OnlineServicesEOS.cpp b/Engine/Plugins/Online/Experimental/OnlineServicesEOS/Source/Private/OnlineServicesEOS.cpp
--- a/Engine/Plugins/Online/Experimental/OnlineServicesEOS/Source/Private/OnlineServicesEOS.cpp
+++ b/Engine/Plugins/Online/Experimental/OnlineServicesEOS/Source/Private/OnlineServicesEOS.cpp
@@ -4,6 +4,7 @@
 #include "OnlineServicesEOSTypes.h"
 
 #include "EOSShared.h"
+#include "EOSSharedTypes.h"
 #include "IEOSSDKManager.h"
 
 #include "AuthEOS.h"
@@ -12,9 +13,101 @@
 #include "OnlineIdEOS.h"
 #include "PresenceEOS.h"
 #include "ExternalUIEOS.h"
+#include "InternetAddrEOS.h"
+#include "Engine/EngineBaseTypes.h"
+#include "NetDriverEOS.h"
 
 namespace UE::Online {
 
+FSocketSubsystemEOSUtils_OnlineServicesEOS::FSocketSubsystemEOSUtils_OnlineServicesEOS(FOnlineServicesEOS* InServicesEOS)
+	: ServicesEOS(InServicesEOS)
+{
+	check(ServicesEOS);
+}
+
+FSocketSubsystemEOSUtils_OnlineServicesEOS::~FSocketSubsystemEOSUtils_OnlineServicesEOS()
+{
+	ServicesEOS = nullptr;
+}
+
+EOS_ProductUserId FSocketSubsystemEOSUtils_OnlineServicesEOS::GetLocalUserId()
+{
+	EOS_ProductUserId Result = nullptr;
+
+	if (ServicesEOS)
+	{
+		using namespace UE::Online;
+
+		IAuthPtr AuthEOS = ServicesEOS->GetAuthInterface();
+		check(AuthEOS);
+
+		FAuthGetAccountByPlatformUserId::Params AuthParams;
+		AuthParams.PlatformUserId = FPlatformMisc::GetPlatformUserForUserIndex(0);
+		TOnlineResult<FAuthGetAccountByPlatformUserId> AuthResult = AuthEOS->GetAccountByPlatformUserId(MoveTemp(AuthParams));
+		if (AuthResult.IsOk())
+		{
+			UE::Online::FAuthGetAccountByPlatformUserId::Result OkValue = AuthResult.GetOkValue();
+
+			Result = GetProductUserIdChecked(OkValue.AccountInfo->UserId);
+		}
+		else
+		{
+			UE_LOG(LogOnlineServices, Verbose, TEXT("[FSocketSubsystemEOSUtils_OnlineServicesEOS::GetLocalUserId] Unable to get account for platform user id [%s]. Error=[%s]"), *ToLogString(AuthParams.PlatformUserId), *AuthResult.GetErrorValue().GetLogString(true));
+		}
+	}
+
+	return Result;
+}
+
+FString FSocketSubsystemEOSUtils_OnlineServicesEOS::GetSessionId()
+{
+	FString Result;
+
+	if (ServicesEOS)
+	{
+		using namespace UE::Online;
+
+		IAuthPtr AuthEOS = ServicesEOS->GetAuthInterface();
+		check(AuthEOS);
+
+		FAuthGetAccountByPlatformUserId::Params AuthParams;
+		AuthParams.PlatformUserId = FPlatformMisc::GetPlatformUserForUserIndex(0);
+		TOnlineResult<FAuthGetAccountByPlatformUserId> AuthResult = AuthEOS->GetAccountByPlatformUserId(MoveTemp(AuthParams));
+		if (AuthResult.IsOk())
+		{
+			FAuthGetAccountByPlatformUserId::Result* AuthOkValue = AuthResult.TryGetOkValue();
+
+			ILobbiesPtr LobbiesEOS = ServicesEOS->GetLobbiesInterface();
+			check(LobbiesEOS);
+
+			FGetJoinedLobbies::Params LobbiesParams;
+			LobbiesParams.LocalUserId = AuthOkValue->AccountInfo->UserId;
+			TOnlineResult<FGetJoinedLobbies> LobbiesResult = LobbiesEOS->GetJoinedLobbies(MoveTemp(LobbiesParams));
+			if (LobbiesResult.IsOk())
+			{
+				FGetJoinedLobbies::Result* LobbiesOkValue = LobbiesResult.TryGetOkValue();
+
+				// TODO: Pending support in Lobbies interface
+				/*for (TSharedRef<FLobby> Lobby : LobbiesOkValue->Lobbies)
+				{
+					 if (Lobby->SessionName == NAME_GameSession)
+						Result = LobbiesEOS->GetLobbyIdString(Lobby);					 
+				}*/
+			}
+			else
+			{
+				UE_LOG(LogOnlineServices, Verbose, TEXT("[FSocketSubsystemEOSUtils_OnlineServicesEOS::GetSessionId] Unable to get joined lobbies for local user id [%s]. Error=[%s]"), *ToLogString(LobbiesParams.LocalUserId), *AuthResult.GetErrorValue().GetLogString(true));
+			}
+		}
+		else
+		{
+			UE_LOG(LogOnlineServices, Verbose, TEXT("[FSocketSubsystemEOSUtils_OnlineServicesEOS::GetSessionId] Unable to get account for platform user id [%s]. Error=[%s]"), *ToLogString(AuthParams.PlatformUserId), *AuthResult.GetErrorValue().GetLogString(true));
+		}
+	}
+
+	return Result;
+}
+
 struct FEOSPlatformConfig
 {
 	FString ProductId;
@@ -94,10 +187,53 @@ void FOnlineServicesEOS::Initialize()
 	}
 
 	EOSPlatformHandle = SDKManager->CreatePlatform(PlatformOptions);
+	if (EOSPlatformHandle)
+	{
+		SocketSubsystem = MakeShareable(new FSocketSubsystemEOS(EOSPlatformHandle, MakeShareable(new FSocketSubsystemEOSUtils_OnlineServicesEOS(this))));
+		if (SocketSubsystem)
+		{
+			FString ErrorStr;
+			if (!SocketSubsystem->Init(ErrorStr))
+			{
+				UE_LOG(LogOnlineServices, Verbose, TEXT("[FOnlineServicesEOS::Initialize] Unable to initialize Socket Subsystem. Error=[%s]"), *ErrorStr);
+			}
+		}
+	}
+	else
+	{
+		UE_LOG(LogOnlineServices, Verbose, TEXT("[FOnlineServicesEOS::Initialize] Unable to initialize Socket Subsystem. EOS Platform Handle was invalid."));
+	}
 
 	FOnlineServicesCommon::Initialize();
 }
 
+TOnlineResult<FGetResolvedConnectString> FOnlineServicesEOS::GetResolvedConnectString(FGetResolvedConnectString::Params&& Params)
+{
+	ILobbiesPtr LobbiesEOS = GetLobbiesInterface();
+	check(LobbiesEOS);
+
+	TOnlineResult<FGetJoinedLobbies> JoinedLobbies = LobbiesEOS->GetJoinedLobbies({ Params.LocalUserId });
+	if (JoinedLobbies.IsOk())
+	{
+		for (TSharedRef<const FLobby>& Lobby : JoinedLobbies.GetOkValue().Lobbies)
+		{
+			if (Lobby->LobbyId == Params.LobbyId)
+			{
+				//It should look like this: "EOS:0002aeeb5b2d4388a3752dd6d31222ec:GameNetDriver:97"
+				FString NetDriverName = GetDefault<UNetDriverEOS>()->NetDriverName.ToString();
+				FInternetAddrEOS TempAddr(GetProductUserIdChecked(Lobby->OwnerAccountId), NetDriverName, GetTypeHash(NetDriverName));
+				return TOnlineResult<FGetResolvedConnectString>({ TempAddr.ToString(true) });
+			}
+		}
+		// No matching lobby
+		return TOnlineResult<FGetResolvedConnectString>(Errors::InvalidParams());
+	}
+	else
+	{
+		return TOnlineResult<FGetResolvedConnectString>(JoinedLobbies.GetErrorValue());
+	}
+}
+
 EOS_HPlatform FOnlineServicesEOS::GetEOSPlatformHandle() const
 {
 	if (EOSPlatformHandle)
diff --git a/Engine/Plugins/Online/Experimental/OnlineServicesEOS/Source/Public/OnlineServicesEOS.h b/Engine/Plugins/Online/Experimental/OnlineServicesEOS/Source/Public/OnlineServicesEOS.h
--- a/Engine/Plugins/Online/Experimental/OnlineServicesEOS/Source/Public/OnlineServicesEOS.h
+++ b/Engine/Plugins/Online/Experimental/OnlineServicesEOS/Source/Public/OnlineServicesEOS.h
@@ -3,6 +3,7 @@
 #pragma once
 
 #include "Online/OnlineServicesCommon.h"
+#include "SocketSubsystemEOS.h"
 
 #if defined(EOS_PLATFORM_BASE_FILE_NAME)
 #include EOS_PLATFORM_BASE_FILE_NAME
@@ -13,6 +14,21 @@ using IEOSPlatformHandlePtr = TSharedPtr<class IEOSPlatformHandle>;
 
 namespace UE::Online {
 
+class FOnlineServicesEOS;
+
+class FSocketSubsystemEOSUtils_OnlineServicesEOS : public ISocketSubsystemEOSUtils
+{
+public:
+	FSocketSubsystemEOSUtils_OnlineServicesEOS(FOnlineServicesEOS* InServicesEOS);
+	virtual ~FSocketSubsystemEOSUtils_OnlineServicesEOS() override;
+
+	virtual EOS_ProductUserId GetLocalUserId() override;
+	virtual FString GetSessionId() override;
+
+private:
+	FOnlineServicesEOS* ServicesEOS;
+};
+
 class ONLINESERVICESEOS_API FOnlineServicesEOS : public FOnlineServicesCommon
 {
 public:
@@ -21,12 +37,15 @@ public:
 	FOnlineServicesEOS();
 	virtual void RegisterComponents() override;
 	virtual void Initialize() override;
+	virtual TOnlineResult<FGetResolvedConnectString> GetResolvedConnectString(FGetResolvedConnectString::Params&& Params) override;
 
 	static EOnlineServices GetServicesProvider() { return EOnlineServices::Epic; }
 
 	EOS_HPlatform GetEOSPlatformHandle() const;
 protected:
 	IEOSPlatformHandlePtr EOSPlatformHandle;
+
+	TSharedPtr<FSocketSubsystemEOS, ESPMode::ThreadSafe> SocketSubsystem;
 };
 
 /* UE::Online */ }
diff --git a/Engine/Plugins/Online/OnlineSubsystemEOS/OnlineSubsystemEOS.uplugin b/Engine/Plugins/Online/OnlineSubsystemEOS/OnlineSubsystemEOS.uplugin
--- a/Engine/Plugins/Online/OnlineSubsystemEOS/OnlineSubsystemEOS.uplugin
+++ b/Engine/Plugins/Online/OnlineSubsystemEOS/OnlineSubsystemEOS.uplugin
@@ -49,6 +49,10 @@
 		{
 			"Name": "EOSVoiceChat",
 			"Enabled": true
+		},
+		{
+			"Name": "SocketSubsystemEOS",
+			"Enabled": true
 		}
 	]
 }
diff --git a/Engine/Plugins/Online/OnlineSubsystemEOS/Source/OnlineSubsystemEOS/OnlineSubsystemEOS.Build.cs b/Engine/Plugins/Online/OnlineSubsystemEOS/Source/OnlineSubsystemEOS/OnlineSubsystemEOS.Build.cs
--- a/Engine/Plugins/Online/OnlineSubsystemEOS/Source/OnlineSubsystemEOS/OnlineSubsystemEOS.Build.cs
+++ b/Engine/Plugins/Online/OnlineSubsystemEOS/Source/OnlineSubsystemEOS/OnlineSubsystemEOS.Build.cs
@@ -27,6 +27,7 @@ public class OnlineSubsystemEOS : ModuleRules
 				"Json",
 				"OnlineSubsystem",
 				"Sockets",
+				"SocketSubsystemEOS",
 				"VoiceChat",
 				"NetCore"
 			}
diff --git a/Engine/Plugins/Online/OnlineSubsystemEOS/Source/OnlineSubsystemEOS/Private/InternetAddrEOS.cpp b/Engine/Plugins/Online/OnlineSubsystemEOS/Source/OnlineSubsystemEOS/Private/InternetAddrEOS.cpp
deleted file mode 100644
diff --git a/Engine/Plugins/Online/OnlineSubsystemEOS/Source/OnlineSubsystemEOS/Private/InternetAddrEOS.h b/Engine/Plugins/Online/OnlineSubsystemEOS/Source/OnlineSubsystemEOS/Private/InternetAddrEOS.h
deleted file mode 100644
diff --git a/Engine/Plugins/Online/OnlineSubsystemEOS/Source/OnlineSubsystemEOS/Private/NetConnectionEOS.cpp b/Engine/Plugins/Online/OnlineSubsystemEOS/Source/OnlineSubsystemEOS/Private/NetConnectionEOS.cpp
deleted file mode 100644
diff --git a/Engine/Plugins/Online/OnlineSubsystemEOS/Source/OnlineSubsystemEOS/Private/NetConnectionEOS.h b/Engine/Plugins/Online/OnlineSubsystemEOS/Source/OnlineSubsystemEOS/Private/NetConnectionEOS.h
deleted file mode 100644
diff --git a/Engine/Plugins/Online/OnlineSubsystemEOS/Source/OnlineSubsystemEOS/Private/NetDriverEOS.cpp b/Engine/Plugins/Online/OnlineSubsystemEOS/Source/OnlineSubsystemEOS/Private/NetDriverEOS.cpp
deleted file mode 100644
diff --git a/Engine/Plugins/Online/OnlineSubsystemEOS/Source/OnlineSubsystemEOS/Private/NetDriverEOS.h b/Engine/Plugins/Online/OnlineSubsystemEOS/Source/OnlineSubsystemEOS/Private/NetDriverEOS.h
deleted file mode 100644
diff --git a/Engine/Plugins/Online/OnlineSubsystemEOS/Source/OnlineSubsystemEOS/Private/OnlineSessionEOS.cpp b/Engine/Plugins/Online/OnlineSubsystemEOS/Source/OnlineSubsystemEOS/Private/OnlineSessionEOS.cpp
--- a/Engine/Plugins/Online/OnlineSubsystemEOS/Source/OnlineSubsystemEOS/Private/OnlineSessionEOS.cpp
+++ b/Engine/Plugins/Online/OnlineSubsystemEOS/Source/OnlineSubsystemEOS/Private/OnlineSessionEOS.cpp
@@ -366,7 +366,7 @@ FOnlineSessionInfoEOS::FOnlineSessionInfoEOS(const FString& InHostIp, FUniqueNet
 {
 	if (InHostIp.StartsWith(EOS_CONNECTION_URL_PREFIX, ESearchCase::IgnoreCase))
 	{
-		HostAddr = ISocketSubsystem::Get(EOS_SUBSYSTEM)->GetAddressFromString(InHostIp);
+		HostAddr = ISocketSubsystem::Get(EOS_SOCKETSUBSYSTEM)->GetAddressFromString(InHostIp);
 		EOSAddress = InHostIp;
 	}
 	else
@@ -493,7 +493,7 @@ void FOnlineSessionEOS::Init(const FString& InBucketId)
 
 	bIsDedicatedServer = IsRunningDedicatedServer();
 	bIsUsingP2PSockets = false;
-	GConfig->GetBool(TEXT("/Script/OnlineSubsystemEOS.NetDriverEOS"), TEXT("bIsUsingP2PSockets"), bIsUsingP2PSockets, GEngineIni);
+	GConfig->GetBool(TEXT("/Script/SocketSubsystemEOS.NetDriverEOS"), TEXT("bIsUsingP2PSockets"), bIsUsingP2PSockets, GEngineIni);
 }
 
 /**
diff --git a/Engine/Plugins/Online/OnlineSubsystemEOS/Source/OnlineSubsystemEOS/Private/OnlineSubsystemEOS.cpp b/Engine/Plugins/Online/OnlineSubsystemEOS/Source/OnlineSubsystemEOS/Private/OnlineSubsystemEOS.cpp
--- a/Engine/Plugins/Online/OnlineSubsystemEOS/Source/OnlineSubsystemEOS/Private/OnlineSubsystemEOS.cpp
+++ b/Engine/Plugins/Online/OnlineSubsystemEOS/Source/OnlineSubsystemEOS/Private/OnlineSubsystemEOS.cpp
@@ -150,7 +150,47 @@ struct FEOSPlatformOptions :
 
 FPlatformEOSHelpersPtr FOnlineSubsystemEOS::EOSHelpersPtr;
 
-bool FCallbackBase::bShouldCancelAllCallbacks = false;
+FSocketSubsystemEOSUtils_OnlineSubsystemEOS::FSocketSubsystemEOSUtils_OnlineSubsystemEOS(FOnlineSubsystemEOS* InSubsystemEOS)
+	: SubsystemEOS(InSubsystemEOS)
+{
+	check(SubsystemEOS != nullptr);
+}
+
+FSocketSubsystemEOSUtils_OnlineSubsystemEOS::~FSocketSubsystemEOSUtils_OnlineSubsystemEOS()
+{
+	SubsystemEOS = nullptr;
+}
+
+EOS_ProductUserId FSocketSubsystemEOSUtils_OnlineSubsystemEOS::GetLocalUserId()
+{
+	EOS_ProductUserId Result = nullptr;
+
+	if (SubsystemEOS)
+	{
+		Result = SubsystemEOS->UserManager->GetLocalProductUserId();
+	}
+
+	return Result;
+}
+
+FString FSocketSubsystemEOSUtils_OnlineSubsystemEOS::GetSessionId()
+{
+	FString Result;
+
+	if (SubsystemEOS)
+	{
+		const IOnlineSessionPtr DefaultSessionInt = SubsystemEOS->GetSessionInterface();
+		if (DefaultSessionInt.IsValid())
+		{
+			if (const FNamedOnlineSession* const NamedSession = DefaultSessionInt->GetNamedSession(NAME_GameSession))
+			{
+				Result = NamedSession->GetSessionIdStr();
+			}
+		}
+	}
+
+	return Result;
+}
 
 void FOnlineSubsystemEOS::ModuleInit()
 {
@@ -350,12 +390,6 @@ bool FOnlineSubsystemEOS::Init()
 		UE_LOG_ONLINE(Error, TEXT("FOnlineSubsystemEOS: failed to init EOS platform, couldn't get achievements handle"));
 		return false;
 	}
-	P2PHandle = EOS_Platform_GetP2PInterface(*EOSPlatformHandle);
-	if (P2PHandle == nullptr)
-	{
-		UE_LOG_ONLINE(Error, TEXT("FOnlineSubsystemEOS: failed to init EOS platform, couldn't get p2p handle"));
-		return false;
-	}
 	// Disable ecom if not part of EGS
 	if (bWasLaunchedByEGS)
 	{
@@ -380,7 +414,7 @@ bool FOnlineSubsystemEOS::Init()
 		return false;
 	}
 
-	SocketSubsystem = MakeShareable(new FSocketSubsystemEOS(this));
+	SocketSubsystem = MakeShareable(new FSocketSubsystemEOS(EOSPlatformHandle, MakeShareable(new FSocketSubsystemEOSUtils_OnlineSubsystemEOS(this))));
 	FString ErrorMessage;
 	SocketSubsystem->Init(ErrorMessage);
 
@@ -557,7 +591,6 @@ FOnlineSubsystemEOS::FOnlineSubsystemEOS(FName InInstanceName) :
 	, LeaderboardsHandle(nullptr)
 	, MetricsHandle(nullptr)
 	, AchievementsHandle(nullptr)
-	, P2PHandle(nullptr)
 	, EcomHandle(nullptr)
 	, TitleStorageHandle(nullptr)
 	, PlayerDataStorageHandle(nullptr)
diff --git a/Engine/Plugins/Online/OnlineSubsystemEOS/Source/OnlineSubsystemEOS/Private/OnlineSubsystemEOSTypes.h b/Engine/Plugins/Online/OnlineSubsystemEOS/Source/OnlineSubsystemEOS/Private/OnlineSubsystemEOSTypes.h
--- a/Engine/Plugins/Online/OnlineSubsystemEOS/Source/OnlineSubsystemEOS/Private/OnlineSubsystemEOSTypes.h
+++ b/Engine/Plugins/Online/OnlineSubsystemEOS/Source/OnlineSubsystemEOS/Private/OnlineSubsystemEOSTypes.h
@@ -8,23 +8,12 @@
 #include "Interfaces/OnlineFriendsInterface.h"
 #include "Interfaces/OnlinePresenceInterface.h"
 #include "Interfaces/OnlineUserInterface.h"
+#include "EOSSharedTypes.h"
 #include "EOSShared.h"
 #include "IPAddress.h"
 
 #include "OnlineSubsystemEOSPackage.h"
 
-#if WITH_EOS_SDK
-	#if defined(EOS_PLATFORM_BASE_FILE_NAME)
-	#include EOS_PLATFORM_BASE_FILE_NAME
-	#endif
-
-	#include "eos_common.h"
-#endif
-
-// Expect URLs to look like "EOS:PUID:SocketName:Channel"
-#define EOS_CONNECTION_URL_PREFIX TEXT("EOS")
-#define EOS_URL_SEPARATOR TEXT(":")
-
 #define EOS_OSS_STRING_BUFFER_LENGTH 256 + 1 // 256 plus null terminator
 
 class FOnlineSubsystemEOS;
@@ -448,17 +437,6 @@ static inline FString MakeNetIdStringFromIds(EOS_EpicAccountId AccountId, EOS_Pr
 	return NetId;
 }
 
-/** Used to store a pointer to the EOS callback object without knowing type */
-class FCallbackBase
-{
-	static bool bShouldCancelAllCallbacks;
-
-public:
-	virtual ~FCallbackBase() {}
-	static bool ShouldCancelAllCallbacks() { return FCallbackBase::bShouldCancelAllCallbacks; }
-	static void CancelAllCallbacks() { FCallbackBase::bShouldCancelAllCallbacks = true; }
-};
-
 /** Class to handle all callbacks generically using a lambda to process callback results */
 template<typename CallbackFuncType, typename CallbackType>
 class TEOSCallback :
@@ -700,41 +678,6 @@ private:
 	}
 };
 
-/** Class to handle all callbacks generically using a lambda to process callback results */
-template<typename CallbackFuncType, typename CallbackType>
-class TEOSGlobalCallback :
-	public FCallbackBase
-{
-public:
-	TFunction<void(const CallbackType*)> CallbackLambda;
-
-	TEOSGlobalCallback() = default;
-	virtual ~TEOSGlobalCallback() = default;
-
-
-	CallbackFuncType GetCallbackPtr()
-	{
-		return &CallbackImpl;
-	}
-
-private:
-	static void EOS_CALL CallbackImpl(const CallbackType* Data)
-	{
-		check(IsInGameThread());
-
-		TEOSGlobalCallback* CallbackThis = (TEOSGlobalCallback*)Data->ClientData;
-		check(CallbackThis);
-
-		if (FCallbackBase::ShouldCancelAllCallbacks())
-		{
-			return;
-		}
-
-		check(CallbackThis->CallbackLambda);
-		CallbackThis->CallbackLambda(Data);
-	}
-};
-
 #include "eos_sessions_types.h"
 
 /**
diff --git a/Engine/Plugins/Online/OnlineSubsystemEOS/Source/OnlineSubsystemEOS/Private/SocketEOS.cpp b/Engine/Plugins/Online/OnlineSubsystemEOS/Source/OnlineSubsystemEOS/Private/SocketEOS.cpp
deleted file mode 100644
diff --git a/Engine/Plugins/Online/OnlineSubsystemEOS/Source/OnlineSubsystemEOS/Private/SocketEOS.h b/Engine/Plugins/Online/OnlineSubsystemEOS/Source/OnlineSubsystemEOS/Private/SocketEOS.h
deleted file mode 100644
diff --git a/Engine/Plugins/Online/OnlineSubsystemEOS/Source/OnlineSubsystemEOS/Private/SocketSubsystemEOS.cpp b/Engine/Plugins/Online/OnlineSubsystemEOS/Source/OnlineSubsystemEOS/Private/SocketSubsystemEOS.cpp
deleted file mode 100644
diff --git a/Engine/Plugins/Online/OnlineSubsystemEOS/Source/OnlineSubsystemEOS/Private/SocketSubsystemEOS.h b/Engine/Plugins/Online/OnlineSubsystemEOS/Source/OnlineSubsystemEOS/Private/SocketSubsystemEOS.h
deleted file mode 100644
diff --git a/Engine/Plugins/Online/OnlineSubsystemEOS/Source/OnlineSubsystemEOS/Public/OnlineSubsystemEOS.h b/Engine/Plugins/Online/OnlineSubsystemEOS/Source/OnlineSubsystemEOS/Public/OnlineSubsystemEOS.h
--- a/Engine/Plugins/Online/OnlineSubsystemEOS/Source/OnlineSubsystemEOS/Public/OnlineSubsystemEOS.h
+++ b/Engine/Plugins/Online/OnlineSubsystemEOS/Source/OnlineSubsystemEOS/Public/OnlineSubsystemEOS.h
@@ -49,6 +49,19 @@ typedef TSharedPtr<class FOnlineUserCloudEOS, ESPMode::ThreadSafe> FOnlineUserCl
 
 typedef TSharedPtr<FPlatformEOSHelpers, ESPMode::ThreadSafe> FPlatformEOSHelpersPtr;
 
+class FSocketSubsystemEOSUtils_OnlineSubsystemEOS : public ISocketSubsystemEOSUtils
+{
+public:
+	FSocketSubsystemEOSUtils_OnlineSubsystemEOS(FOnlineSubsystemEOS* InSubsystemEOS);
+	virtual ~FSocketSubsystemEOSUtils_OnlineSubsystemEOS() override;
+
+	virtual EOS_ProductUserId GetLocalUserId() override;
+	virtual FString GetSessionId() override;
+
+private:
+	FOnlineSubsystemEOS* SubsystemEOS;
+};
+
 /**
  *	OnlineSubsystemEOS - Implementation of the online subsystem for EOS services
  */
@@ -128,7 +141,6 @@ PACKAGE_SCOPE:
 	EOS_HLeaderboards LeaderboardsHandle;
 	EOS_HMetrics MetricsHandle;
 	EOS_HAchievements AchievementsHandle;
-	EOS_HP2P P2PHandle;
 	EOS_HEcom EcomHandle;
 	EOS_HTitleStorage TitleStorageHandle;
 	EOS_HPlayerDataStorage PlayerDataStorageHandle;
diff --git a/Engine/Plugins/Online/SocketSubsystemEOS/SocketSubsystemEOS.uplugin b/Engine/Plugins/Online/SocketSubsystemEOS/SocketSubsystemEOS.uplugin
new file mode 100644
--- /dev/null
+++ b/Engine/Plugins/Online/SocketSubsystemEOS/SocketSubsystemEOS.uplugin
@@ -0,0 +1,36 @@
+{
+	"FileVersion": 3,
+
+	"FriendlyName": "Socket Subsystem EOS",
+	"Version": 1,
+	"VersionName": "1.0",
+	"Description": "Responsible for management of EOS P2P Socket connections.",
+	"Category": "Online Platform",
+	"CreatedBy": "Epic Games, Inc.",
+	"CreatedByURL": "http://epicgames.com",
+	"EnabledByDefault": false,
+
+	"Modules": [
+		{
+			"Name": "SocketSubsystemEOS",
+			"Type": "RuntimeNoCommandlet",
+			"LoadingPhase": "PostConfigInit",
+			"PlatformDenyList": [
+				"LinuxArm64",
+				"Hololens",
+				"TVOS"
+			]
+		}
+	],
+
+	"Plugins": [
+		{
+			"Name": "OnlineSubsystemUtils",
+			"Enabled": true
+		},
+		{
+			"Name": "EOSShared",
+			"Enabled": true
+		}
+	]
+}
diff --git a/Engine/Plugins/Online/SocketSubsystemEOS/Source/SocketSubsystemEOS/Private/InternetAddrEOS.cpp b/Engine/Plugins/Online/SocketSubsystemEOS/Source/SocketSubsystemEOS/Private/InternetAddrEOS.cpp
new file mode 100644
--- /dev/null
+++ b/Engine/Plugins/Online/SocketSubsystemEOS/Source/SocketSubsystemEOS/Private/InternetAddrEOS.cpp
@@ -0,0 +1,184 @@
+// Copyright Epic Games, Inc. All Rights Reserved.
+
+#include "InternetAddrEOS.h"
+#include "Online/CoreOnline.h"
+#include "EOSSharedTypes.h"
+#include "Containers/StringView.h"
+
+DEFINE_LOG_CATEGORY(LogSocketSubsystemEOS);
+
+inline uint8 PortToChannel(int32 InPort)
+{
+	return InPort > 255 ? InPort % 256 : FMath::Clamp(InPort, 0, 255);
+}
+
+FInternetAddrEOS::FInternetAddrEOS()
+	: LocalUserId(nullptr)
+	, RemoteUserId(nullptr)
+	, Channel(0)
+{
+	SocketName[0] = '\0';
+}
+
+FInternetAddrEOS::FInternetAddrEOS(const FString& InRemoteUserId, const FString& InSocketName, const int32 InChannel)
+	: LocalUserId(nullptr)
+	, RemoteUserId(nullptr)
+{
+#if WITH_EOS_SDK
+	RemoteUserId = EOS_ProductUserId_FromString(TCHAR_TO_UTF8(*InRemoteUserId));
+#endif
+	FCStringAnsi::Strcpy(SocketName, TCHAR_TO_UTF8(*InSocketName));
+	Channel = PortToChannel(InChannel);
+}
+
+#if WITH_EOS_SDK
+FInternetAddrEOS::FInternetAddrEOS(const EOS_ProductUserId InRemoteUserId, const FString& InSocketName, const int32 InChannel)
+	: LocalUserId(nullptr)
+	, RemoteUserId(InRemoteUserId)
+	, Channel(InChannel)
+{
+	FCStringAnsi::Strcpy(SocketName, TCHAR_TO_UTF8(*InSocketName));
+	Channel = PortToChannel(InChannel);
+}
+#endif
+
+void FInternetAddrEOS::SetIp(uint32)
+{
+	UE_LOG(LogSocketSubsystemEOS, Warning, TEXT("Calls to FInternetAddrEOS::SetIp are not valid"));
+}
+
+void FInternetAddrEOS::SetIp(const TCHAR* InAddr, bool& bIsValid)
+{
+	bIsValid = false;
+
+	if (InAddr == nullptr)
+	{
+		return;
+	}
+
+	Channel = 0;
+
+	TArray<FString> UrlParts;
+	FString FullAddress = InAddr;
+	FullAddress.ParseIntoArray(UrlParts, EOS_URL_SEPARATOR, false);
+	// Expect URLs to look like "EOS:PUID:SocketName:Channel" and channel can be optional
+	if (UrlParts.Num() < 3 || UrlParts.Num() > 4)
+	{
+		return;
+	}
+	if (UrlParts[0] != EOS_CONNECTION_URL_PREFIX)
+	{
+		return;
+	}
+#if WITH_EOS_SDK
+	RemoteUserId = EOS_ProductUserId_FromString(TCHAR_TO_UTF8(*UrlParts[1]));
+	if (EOS_ProductUserId_IsValid(RemoteUserId) == EOS_FALSE)
+#endif
+	{
+		return;
+	}
+	if (UrlParts[2].Len() == 0)
+	{
+		return;
+	}
+	SetSocketName(UrlParts[2]);
+	if (UrlParts.Num() == 4)
+	{
+		Channel = FCString::Atoi(*UrlParts[3]);
+	}
+	bIsValid = true;
+}
+
+void FInternetAddrEOS::GetIp(uint32& OutAddr) const
+{
+	OutAddr = 0u;
+
+	UE_LOG(LogSocketSubsystemEOS, Warning, TEXT("Calls to FInternetAddrEOS::GetIp are not valid"));
+}
+
+void FInternetAddrEOS::SetPort(int32 InPort)
+{
+	Channel = PortToChannel(InPort);
+}
+
+int32 FInternetAddrEOS::GetPort() const
+{
+	return Channel;
+}
+
+void FInternetAddrEOS::SetRawIp(const TArray<uint8>& RawAddr)
+{
+	// Need auto here, as might give us different return type depending on size of TCHAR
+	auto ConvertedTCHARData = StringCast<TCHAR>(reinterpret_cast<const ANSICHAR*>(RawAddr.GetData()), RawAddr.Num());
+	const FString IpAsString(ConvertedTCHARData.Length(), ConvertedTCHARData.Get());
+
+	bool bUnused;
+	SetIp(*IpAsString, bUnused);
+}
+
+TArray<uint8> FInternetAddrEOS::GetRawIp() const
+{
+	// We could do this more efficiently, but this was much faster to write.
+	const FString StringVersion = ToString(true);
+
+	// Need auto here, as might give us different return type depending on size of TCHAR
+	auto ConvertedANSIData = StringCast<ANSICHAR>(*StringVersion, StringVersion.Len());
+
+	TArray<uint8> OutData;
+	for (int32 Index = 0; Index < ConvertedANSIData.Length(); ++Index)
+	{
+		OutData.Add(ConvertedANSIData.Get()[Index]);
+	}
+
+	return OutData;
+}
+
+void FInternetAddrEOS::SetAnyAddress()
+{
+}
+
+void FInternetAddrEOS::SetBroadcastAddress()
+{
+}
+
+void FInternetAddrEOS::SetLoopbackAddress()
+{
+}
+
+FString FInternetAddrEOS::ToString(bool bAppendPort) const
+{
+	char PuidBuffer[64];
+	int32 BufferLen = 64;
+#if WITH_EOS_SDK
+	if (EOS_ProductUserId_ToString(RemoteUserId, PuidBuffer, &BufferLen) != EOS_EResult::EOS_Success)
+#endif
+	{
+		PuidBuffer[0] = '\0';
+	}
+
+	if (bAppendPort)
+	{
+		return FString::Printf(TEXT("%s%s%s%s%s%s%u"), EOS_CONNECTION_URL_PREFIX, EOS_URL_SEPARATOR, UTF8_TO_TCHAR(PuidBuffer), EOS_URL_SEPARATOR, UTF8_TO_TCHAR(SocketName), EOS_URL_SEPARATOR, Channel);
+	}
+
+	return FString::Printf(TEXT("%s%s%s%s%s"), EOS_CONNECTION_URL_PREFIX, EOS_URL_SEPARATOR, UTF8_TO_TCHAR(PuidBuffer), EOS_URL_SEPARATOR, UTF8_TO_TCHAR(SocketName));
+}
+
+uint32 FInternetAddrEOS::GetTypeHash() const
+{
+	return HashCombine(HashCombine(HashCombine(::GetTypeHash((void*)LocalUserId), ::GetTypeHash((void*)RemoteUserId)), ::GetTypeHash(FAnsiStringView(SocketName, EOS_SOCKET_NAME_SIZE))), Channel);
+}
+
+bool FInternetAddrEOS::IsValid() const
+{
+#if WITH_EOS_SDK
+	return (EOS_ProductUserId_IsValid(LocalUserId) == EOS_TRUE || EOS_ProductUserId_IsValid(RemoteUserId) == EOS_TRUE) && FCStringAnsi::Strlen(SocketName) > 0;
+#else
+	return false;
+#endif
+}
+
+TSharedRef<FInternetAddr> FInternetAddrEOS::Clone() const
+{
+	return MakeShared<FInternetAddrEOS>(*this);
+}
diff --git a/Engine/Plugins/Online/SocketSubsystemEOS/Source/SocketSubsystemEOS/Private/NetConnectionEOS.cpp b/Engine/Plugins/Online/SocketSubsystemEOS/Source/SocketSubsystemEOS/Private/NetConnectionEOS.cpp
new file mode 100644
--- /dev/null
+++ b/Engine/Plugins/Online/SocketSubsystemEOS/Source/SocketSubsystemEOS/Private/NetConnectionEOS.cpp
@@ -0,0 +1,76 @@
+// Copyright Epic Games, Inc. All Rights Reserved.
+
+#include "NetConnectionEOS.h"
+#include "NetDriverEOS.h"
+#include "InternetAddrEOS.h"
+#include "SocketEOS.h"
+
+UNetConnectionEOS::UNetConnectionEOS(const FObjectInitializer& ObjectInitializer)
+	: Super(ObjectInitializer)
+	, bIsPassthrough(false)
+	, bHasP2PSession(false)
+{
+}
+
+void UNetConnectionEOS::InitLocalConnection(UNetDriver* InDriver, FSocket* InSocket, const FURL& InURL, EConnectionState InState, int32 InMaxPacket, int32 InPacketOverhead)
+{
+	bIsPassthrough = !static_cast<UNetDriverEOS*>(InDriver)->bIsUsingP2PSockets || !InURL.Host.StartsWith(EOS_CONNECTION_URL_PREFIX, ESearchCase::IgnoreCase);
+	bHasP2PSession = !bIsPassthrough;
+	  
+	if (bHasP2PSession)
+	{
+		DisableAddressResolution();
+	}
+
+	Super::InitLocalConnection(InDriver, InSocket, InURL, InState, InMaxPacket, InPacketOverhead);
+
+	// We re-set the original port, as if might have been overwritten by Super::InitLocalConnection
+	if (bHasP2PSession)
+	{
+		RemoteAddr->SetPort(InSocket->GetPortNo());
+	}
+}
+
+void UNetConnectionEOS::InitRemoteConnection(UNetDriver* InDriver, FSocket* InSocket, const FURL& InURL, const FInternetAddr& InRemoteAddr, EConnectionState InState, int32 InMaxPacket, int32 InPacketOverhead)
+{
+	bIsPassthrough = static_cast<UNetDriverEOS*>(InDriver)->bIsPassthrough;
+	bHasP2PSession = !bIsPassthrough;
+
+	if (bHasP2PSession)
+	{
+		DisableAddressResolution();
+	}
+
+	Super::InitRemoteConnection(InDriver, InSocket, InURL, InRemoteAddr, InState, InMaxPacket, InPacketOverhead);
+}
+
+void UNetConnectionEOS::CleanUp()
+{
+	Super::CleanUp();
+
+	if (bHasP2PSession)
+	{
+		DestroyEOSConnection();
+	}
+}
+
+void UNetConnectionEOS::DestroyEOSConnection()
+{
+	if (!Socket)
+	{
+		return;
+	}
+
+	if (!bHasP2PSession)
+	{
+		return;
+	}
+
+	bHasP2PSession = false;
+
+	TSharedPtr<FInternetAddrEOS> RemoteAddrEOS = StaticCastSharedPtr<FInternetAddrEOS>(RemoteAddr);
+	if (RemoteAddrEOS.IsValid())
+	{
+		static_cast<FSocketEOS*>(Socket)->Close(*RemoteAddrEOS);
+	}
+}
diff --git a/Engine/Plugins/Online/SocketSubsystemEOS/Source/SocketSubsystemEOS/Private/NetConnectionEOS.h b/Engine/Plugins/Online/SocketSubsystemEOS/Source/SocketSubsystemEOS/Private/NetConnectionEOS.h
new file mode 100644
--- /dev/null
+++ b/Engine/Plugins/Online/SocketSubsystemEOS/Source/SocketSubsystemEOS/Private/NetConnectionEOS.h
@@ -0,0 +1,31 @@
+// Copyright Epic Games, Inc. All Rights Reserved.
+
+#pragma once
+
+#include "CoreMinimal.h"
+#include "IpConnection.h"
+#include "NetConnectionEOS.generated.h"
+
+UCLASS(Transient, Config=Engine)
+class SOCKETSUBSYSTEMEOS_API UNetConnectionEOS
+	: public UIpConnection
+{
+	GENERATED_BODY()
+
+public:
+	explicit UNetConnectionEOS(const FObjectInitializer& ObjectInitializer);
+
+//~ Begin NetConnection Interface
+	virtual void InitLocalConnection(UNetDriver* InDriver, FSocket* InSocket, const FURL& InURL, EConnectionState InState, int32 InMaxPacket = 0, int32 InPacketOverhead = 0) override;
+	virtual void InitRemoteConnection(UNetDriver* InDriver, FSocket* InSocket, const FURL& InURL, const FInternetAddr& InRemoteAddr, EConnectionState InState, int32 InMaxPacket = 0, int32 InPacketOverhead = 0) override;
+	virtual void CleanUp() override;
+//~ End NetConnection Interface
+
+	void DestroyEOSConnection();
+
+public:
+	bool bIsPassthrough;
+
+protected:
+	bool bHasP2PSession;
+};
diff --git a/Engine/Plugins/Online/SocketSubsystemEOS/Source/SocketSubsystemEOS/Private/NetDriverEOS.cpp b/Engine/Plugins/Online/SocketSubsystemEOS/Source/SocketSubsystemEOS/Private/NetDriverEOS.cpp
new file mode 100644
--- /dev/null
+++ b/Engine/Plugins/Online/SocketSubsystemEOS/Source/SocketSubsystemEOS/Private/NetDriverEOS.cpp
@@ -0,0 +1,226 @@
+// Copyright Epic Games, Inc. All Rights Reserved.
+
+#include "NetDriverEOS.h"
+#include "NetConnectionEOS.h"
+#include "SocketEOS.h"
+#include "SocketSubsystemEOS.h"
+#include "Misc/EngineVersionComparison.h"
+#include "EOSSharedTypes.h"
+#include "Engine/Engine.h"
+
+bool UNetDriverEOS::IsAvailable() const
+{
+	// Use passthrough sockets if we are a dedicated server
+	if (IsRunningDedicatedServer())
+	{
+		return false;
+	}
+
+	if (ISocketSubsystem* SocketSubsystem = ISocketSubsystem::Get(EOS_SOCKETSUBSYSTEM))
+	{
+		return true;
+	}
+
+	return false;
+}
+
+bool UNetDriverEOS::InitBase(bool bInitAsClient, FNetworkNotify* InNotify, const FURL& URL, bool bReuseAddressAndPort, FString& Error)
+{
+	if (bIsPassthrough)
+	{
+		UE_LOG(LogTemp, Verbose, TEXT("Running as pass-through"));
+		return Super::InitBase(bInitAsClient, InNotify, URL, bReuseAddressAndPort, Error);
+	}
+
+	if (!UNetDriver::InitBase(bInitAsClient, InNotify, URL, bReuseAddressAndPort, Error))
+	{
+		UE_LOG(LogTemp, Warning, TEXT("Failed to init driver base"));
+		return false;
+	}
+
+	FSocketSubsystemEOS* const SocketSubsystem = static_cast<FSocketSubsystemEOS*>(GetSocketSubsystem());
+	if (!SocketSubsystem)
+	{
+		UE_LOG(LogTemp, Warning, TEXT("Could not get socket subsystem"));
+		return false;
+	}
+
+	// We don't care if our world is null, everything we uses handles it fine
+	const UWorld* const MyWorld = FindWorld();
+
+	// Get our local address (proves we're logged in)
+	TSharedRef<FInternetAddr> LocalAddress = SocketSubsystem->GetLocalBindAddr(MyWorld, *GLog);
+	if (!LocalAddress->IsValid())
+	{
+		// Not logged in?
+		Error = TEXT("Could not bind local address");
+		UE_LOG(LogTemp, Warning, TEXT("Could not bind local address"));
+		return false;
+	}
+
+	// Create our socket
+	SetSocketAndLocalAddress(SocketSubsystem->CreateSocket(NAME_DGram, TEXT("UE4"), NAME_None));
+	if (GetSocket() == nullptr)
+	{
+		UE_LOG(LogTemp, Warning, TEXT("Could not create socket"));
+		return false;
+	}
+
+	// Store our local address and set our port
+	TSharedRef<FInternetAddrEOS> EOSLocalAddress = StaticCastSharedRef<FInternetAddrEOS>(LocalAddress);
+	// Because some platforms remap ports, we will use the ID of the name of the net driver to be our channel
+	EOSLocalAddress->SetChannel(GetTypeHash(NetDriverName.ToString()));
+	// Set our net driver name so we don't accept connections across net driver types
+	EOSLocalAddress->SetSocketName(NetDriverName.ToString());
+
+	static_cast<FSocketEOS*>(GetSocket())->SetLocalAddress(*EOSLocalAddress);
+
+	LocalAddr = LocalAddress;
+
+	return true;
+}
+
+bool UNetDriverEOS::InitConnect(FNetworkNotify* InNotify, const FURL& ConnectURL, FString& Error)
+{
+	if (!bIsUsingP2PSockets || !IsAvailable() || !ConnectURL.Host.StartsWith(EOS_CONNECTION_URL_PREFIX, ESearchCase::IgnoreCase))
+	{
+		UE_LOG(LogTemp, Verbose, TEXT("Connecting using IPNetDriver passthrough. ConnectUrl = (%s)"), *ConnectURL.ToString());
+
+		bIsPassthrough = true;
+		return Super::InitConnect(InNotify, ConnectURL, Error);
+	}
+
+	bool bIsValid = false;
+	TSharedRef<FInternetAddrEOS> RemoteHost = MakeShared<FInternetAddrEOS>();
+	RemoteHost->SetIp(*ConnectURL.Host, bIsValid);
+	if (!bIsValid || ConnectURL.Port < 0)
+	{
+		Error = TEXT("Invalid remote address");
+		UE_LOG(LogTemp, Warning, TEXT("Invalid Remote Address. ConnectUrl = (%s)"), *ConnectURL.ToString());
+		return false;
+	}
+
+	UE_LOG(LogTemp, Verbose, TEXT("Connecting using EOSNetDriver. ConnectUrl = (%s)"), *ConnectURL.ToString());
+
+	if (!InitBase(true, InNotify, ConnectURL, false, Error))
+	{
+		return false;
+	}
+
+	// Set the address to what was parsed (us + remote)
+	LocalAddr = RemoteHost;
+
+	// Reference to our newly created socket
+	FSocket* CurSocket = GetSocket();
+
+	// Bind our local port
+	FSocketSubsystemEOS* const SocketSubsystem = static_cast<FSocketSubsystemEOS*>(GetSocketSubsystem());
+	check(SocketSubsystem);
+	if (!SocketSubsystem->BindNextPort(CurSocket, *LocalAddr, MaxPortCountToTry + 1, 1))
+	{
+		// Failure
+		Error = TEXT("Could not bind local port");
+		UE_LOG(LogTemp, Warning, TEXT("Could not bind local port in %d attempts"), MaxPortCountToTry);
+		return false;
+	}
+
+	// Create an unreal connection to the server
+	UNetConnectionEOS* Connection = NewObject<UNetConnectionEOS>(NetConnectionClass);
+	check(Connection);
+
+	// Set it as the server connection before anything else so everything knows this is a client
+	ServerConnection = Connection;
+	Connection->InitLocalConnection(this, CurSocket, ConnectURL, USOCK_Pending);
+
+	CreateInitialClientChannels();
+
+	return true;
+}
+
+bool UNetDriverEOS::InitListen(FNetworkNotify* InNotify, FURL& LocalURL, bool bReuseAddressAndPort, FString& Error)
+{
+	if (!bIsUsingP2PSockets || !IsAvailable() || LocalURL.HasOption(TEXT("bIsLanMatch")) || LocalURL.HasOption(TEXT("bUseIPSockets")))
+	{
+		UE_LOG(LogTemp, Verbose, TEXT("Init as IPNetDriver listen server. LocalURL = (%s)"), *LocalURL.ToString());
+
+		bIsPassthrough = true;
+		return Super::InitListen(InNotify, LocalURL, bReuseAddressAndPort, Error);
+	}
+
+	UE_LOG(LogTemp, Verbose, TEXT("Init as EOSNetDriver listen server. LocalURL = (%s)"), *LocalURL.ToString());
+
+	if (!InitBase(false, InNotify, LocalURL, bReuseAddressAndPort, Error))
+	{
+		return false;
+	}
+
+	// Bind our specified port if provided
+	FSocket* CurSocket = GetSocket();
+	if (!CurSocket->Listen(0))
+	{
+		Error = TEXT("Could not listen");
+		UE_LOG(LogTemp, Warning, TEXT("Could not listen on socket"));
+		return false;
+	}
+
+	InitConnectionlessHandler();
+
+	UE_LOG(LogTemp, Verbose, TEXT("Initialized as an EOSP2P listen server"));
+	return true;
+}
+
+ISocketSubsystem* UNetDriverEOS::GetSocketSubsystem()
+{
+	return ISocketSubsystem::Get(bIsPassthrough ? PLATFORM_SOCKETSUBSYSTEM : EOS_SOCKETSUBSYSTEM);
+}
+
+void UNetDriverEOS::Shutdown()
+{
+	UE_LOG(LogTemp, Verbose, TEXT("Shutting down NetDriver"));
+
+	Super::Shutdown();
+
+	// Kill our P2P sessions now, instead of when garbage collection kicks in later
+	if (!bIsPassthrough)
+	{
+		if (UNetConnectionEOS* const EOSServerConnection = Cast<UNetConnectionEOS>(ServerConnection))
+		{
+			EOSServerConnection->DestroyEOSConnection();
+		}
+		for (UNetConnection* Client : ClientConnections)
+		{
+			if (UNetConnectionEOS* const EOSClient = Cast<UNetConnectionEOS>(Client))
+			{
+				EOSClient->DestroyEOSConnection();
+			}
+		}
+	}
+}
+
+int UNetDriverEOS::GetClientPort()
+{
+	if (bIsPassthrough)
+	{
+		return Super::GetClientPort();
+	}
+
+	// Starting range of dynamic/private/ephemeral ports
+	return 49152;
+}
+
+UWorld* UNetDriverEOS::FindWorld() const
+{
+	UWorld* MyWorld = GetWorld();
+
+	// If we don't have a world, we may be a pending net driver
+	if (!MyWorld && GEngine)
+	{
+		if (FWorldContext* WorldContext = GEngine->GetWorldContextFromPendingNetGameNetDriver(this))
+		{
+			MyWorld = WorldContext->World();
+		}
+	}
+
+	return MyWorld;
+}
+
diff --git a/Engine/Plugins/Online/SocketSubsystemEOS/Source/SocketSubsystemEOS/Private/SocketEOS.cpp b/Engine/Plugins/Online/SocketSubsystemEOS/Source/SocketSubsystemEOS/Private/SocketEOS.cpp
new file mode 100644
--- /dev/null
+++ b/Engine/Plugins/Online/SocketSubsystemEOS/Source/SocketSubsystemEOS/Private/SocketEOS.cpp
@@ -0,0 +1,676 @@
+// Copyright Epic Games, Inc. All Rights Reserved.
+
+#include "SocketEOS.h"
+#include "SocketTypes.h"
+#include "SocketSubsystemEOS.h"
+
+#if WITH_EOS_SDK
+	#include "eos_p2p.h"
+#endif
+
+#if WANTS_NP_LOGGING
+
+#include "Windows/AllowWindowsPlatformTypes.h"
+THIRD_PARTY_INCLUDES_START
+#include <Windows.h>
+THIRD_PARTY_INCLUDES_END
+
+const TCHAR* GetLogPrefix()
+{
+	static FString Prefix;
+	if (Prefix.Len() == 0)
+	{
+		FParse::Value(FCommandLine::Get(), TEXT("LogPrefix="), Prefix);
+		if (Prefix.Len() == 0)
+		{
+			Prefix = TEXT("Unknown");
+		}
+	}
+	return *Prefix;
+}
+
+void NpLog(const TCHAR* Msg)
+{
+	static HWND EditWindow = NULL;
+	// Get the edit window so we can send messages to it
+	if (EditWindow == NULL)
+	{
+		HWND MainWindow = FindWindowW(NULL, L"Untitled - Notepad");
+		if (MainWindow == NULL)
+		{
+			MainWindow = FindWindowW(NULL, L"*Untitled - Notepad");
+		}
+		if (MainWindow != NULL)
+		{
+			EditWindow = FindWindowExW(MainWindow, NULL, L"Edit", NULL);
+		}
+	}
+	if (EditWindow != NULL)
+	{
+		SendMessageW(EditWindow, EM_REPLACESEL, TRUE, (LPARAM)Msg);
+	}
+}
+
+#include "Windows/HideWindowsPlatformTypes.h"
+#endif
+
+FSocketEOS::FSocketEOS(FSocketSubsystemEOS& InSocketSubsystem, const FString& InSocketDescription)
+	: FSocket(ESocketType::SOCKTYPE_Datagram, InSocketDescription, NAME_None)
+	, SocketSubsystem(InSocketSubsystem)
+	, bIsListening(false)
+#if WITH_EOS_SDK
+	, ConnectNotifyCallback(nullptr)
+	, ConnectNotifyId(EOS_INVALID_NOTIFICATIONID)
+	, ClosedNotifyCallback(nullptr)
+	, ClosedNotifyId(EOS_INVALID_NOTIFICATIONID)
+#endif
+{
+}
+
+FSocketEOS::~FSocketEOS()
+{
+	Close();
+
+	if (LocalAddress.IsValid())
+	{
+		SocketSubsystem.UnbindChannel(LocalAddress);
+		LocalAddress = FInternetAddrEOS();
+	}
+}
+
+bool FSocketEOS::Shutdown(ESocketShutdownMode Mode)
+{
+	SocketSubsystem.SetLastSocketError(ESocketErrors::SE_EOPNOTSUPP);
+	return false;
+}
+
+bool FSocketEOS::Close()
+{
+	check(IsInGameThread() && "p2p does not support multithreading");
+
+#if WITH_EOS_SDK
+	if (ConnectNotifyId != EOS_INVALID_NOTIFICATIONID)
+	{
+		EOS_P2P_RemoveNotifyPeerConnectionRequest(SocketSubsystem.GetP2PHandle(), ConnectNotifyId);
+	}
+	delete ConnectNotifyCallback;
+	ConnectNotifyCallback = nullptr;
+	if (ClosedNotifyId != EOS_INVALID_NOTIFICATIONID)
+	{
+		EOS_P2P_RemoveNotifyPeerConnectionClosed(SocketSubsystem.GetP2PHandle(), ClosedNotifyId);
+	}
+	delete ClosedNotifyCallback;
+	ClosedNotifyCallback = nullptr;
+
+	if (LocalAddress.IsValid())
+	{
+		EOS_P2P_SocketId SocketId = { };
+		SocketId.ApiVersion = EOS_P2P_SOCKETID_API_LATEST;
+		FCStringAnsi::Strcpy(SocketId.SocketName, LocalAddress.GetSocketName());
+
+		EOS_P2P_CloseConnectionsOptions Options = { };
+		Options.ApiVersion = EOS_P2P_CLOSECONNECTIONS_API_LATEST;
+		Options.LocalUserId = SocketSubsystem.GetLocalUserId();
+		Options.SocketId = &SocketId;
+
+		EOS_EResult Result = EOS_P2P_CloseConnections(SocketSubsystem.GetP2PHandle(), &Options);
+
+		UE_LOG(LogSocketSubsystemEOS, Log, TEXT("Closing socket (%s) with result (%s)"), *LocalAddress.ToString(true), ANSI_TO_TCHAR(EOS_EResult_ToString(Result)));
+		NP_LOG(TEXT("[%s] - Closing socket (%s) with result (%s)\r\n"), GetLogPrefix(), *LocalAddress.ToString(true), ANSI_TO_TCHAR(EOS_EResult_ToString(Result)));
+
+		ClosedRemotes.Empty();
+	}
+#endif
+	return true;
+}
+
+bool FSocketEOS::Bind(const FInternetAddr& Addr)
+{
+	check(IsInGameThread() && "p2p does not support multithreading");
+
+	if (!Addr.IsValid())
+	{
+		UE_LOG(LogSocketSubsystemEOS, Warning, TEXT("Attempted to bind to invalid address. Address = (%s)"), *Addr.ToString(true));
+		SocketSubsystem.SetLastSocketError(ESocketErrors::SE_EADDRNOTAVAIL);
+		return false;
+	}
+
+	// Ensure we called Initialize so we know who we are
+	if (LocalAddress.GetRemoteUserId() != nullptr)
+	{
+		UE_LOG(LogSocketSubsystemEOS, Warning, TEXT("Attempted to bind on a socket that was not initialized. Address = (%s)"), *Addr.ToString(true));
+		SocketSubsystem.SetLastSocketError(ESocketErrors::SE_NOTINITIALISED);
+		return false;
+	}
+
+	// If we have a remote user id, we're already bound
+	if (LocalAddress.GetRemoteUserId() != nullptr)
+	{
+		UE_LOG(LogSocketSubsystemEOS, Warning, TEXT("Attempted to bind a socket that was already bound. ExistingAddress = (%s) NewAddress = (%s)"), *LocalAddress.ToString(true), *Addr.ToString(true));
+		SocketSubsystem.SetLastSocketError(ESocketErrors::SE_EADDRINUSE);
+		return false;
+	}
+
+	const FInternetAddrEOS& EOSAddr = static_cast<const FInternetAddrEOS&>(Addr);
+	if (!SocketSubsystem.BindChannel(EOSAddr))
+	{
+		UE_LOG(LogSocketSubsystemEOS, Warning, TEXT("Attempted to bind a socket to a port in use. NewAddress = (%s)"), *Addr.ToString(true));
+		// BindPort sets our LastSocketError
+		return false;
+	}
+
+#if WITH_EOS_SDK
+	EOS_ProductUserId LocalUserId = LocalAddress.GetLocalUserId();
+#else
+	void* LocalUserId = LocalAddress.GetLocalUserId();
+#endif
+	LocalAddress = EOSAddr;
+	LocalAddress.SetLocalUserId(LocalUserId);
+
+	UE_LOG(LogSocketSubsystemEOS, Verbose, TEXT("Successfully bound socket to address (%s)"), *LocalAddress.ToString(true));
+	NP_LOG(TEXT("[%s] - Successfully bound socket to address (%s)\r\n"), GetLogPrefix(), *LocalAddress.ToString(true));
+	return true;
+}
+
+bool FSocketEOS::Connect(const FInternetAddr& Addr)
+{
+	/** Not supported - connectionless (UDP) only */
+	SocketSubsystem.SetLastSocketError(ESocketErrors::SE_EOPNOTSUPP);
+	return false;
+}
+
+bool FSocketEOS::Listen(int32)
+{
+	check(IsInGameThread() && "p2p does not support multithreading");
+
+	if (!LocalAddress.IsValid())
+	{
+		UE_LOG(LogSocketSubsystemEOS, Warning, TEXT("Attempted to listen without a bound address. Address = (%s)"), *LocalAddress.ToString(true));
+		SocketSubsystem.SetLastSocketError(ESocketErrors::SE_EADDRINUSE);
+		return false;
+	}
+
+#if WITH_EOS_SDK
+	// Add listener for inbound connections
+	EOS_P2P_SocketId SocketId = { };
+	SocketId.ApiVersion = EOS_P2P_SOCKETID_API_LATEST;
+	FCStringAnsi::Strcpy(SocketId.SocketName, LocalAddress.GetSocketName());
+
+	EOS_P2P_AddNotifyPeerConnectionRequestOptions Options = { };
+	Options.ApiVersion = EOS_P2P_ADDNOTIFYPEERCONNECTIONREQUEST_API_LATEST;
+	Options.LocalUserId = LocalAddress.GetLocalUserId();
+	Options.SocketId = &SocketId;
+
+	ConnectNotifyCallback = new FConnectNotifyCallback();
+	ConnectNotifyCallback->CallbackLambda = [this](const EOS_P2P_OnIncomingConnectionRequestInfo* Info)
+	{
+		char PuidBuffer[64];
+		int32 BufferLen = 64;
+		if (EOS_ProductUserId_ToString(Info->RemoteUserId, PuidBuffer, &BufferLen) != EOS_EResult::EOS_Success)
+		{
+			PuidBuffer[0] = '\0';
+		}
+		FString RemoteUser(PuidBuffer);
+
+		if (Info->LocalUserId == LocalAddress.GetLocalUserId() && FCStringAnsi::Stricmp(Info->SocketId->SocketName, LocalAddress.GetSocketName()) == 0)
+		{
+			// In case they disconnected and then reconnected, remove them from our closed list
+			FInternetAddrEOS RemoteAddress(Info->RemoteUserId, Info->SocketId->SocketName, LocalAddress.GetChannel());
+			RemoteAddress.SetLocalUserId(LocalAddress.GetLocalUserId());
+			ClosedRemotes.Remove(RemoteAddress);
+
+			EOS_P2P_SocketId SocketId = { };
+			SocketId.ApiVersion = EOS_P2P_SOCKETID_API_LATEST;
+			FCStringAnsi::Strcpy(SocketId.SocketName, Info->SocketId->SocketName);
+
+			EOS_P2P_AcceptConnectionOptions Options = { };
+			Options.ApiVersion = EOS_P2P_ACCEPTCONNECTION_API_LATEST;
+			Options.LocalUserId = LocalAddress.GetLocalUserId();
+			Options.RemoteUserId = Info->RemoteUserId;
+			Options.SocketId = &SocketId;
+			EOS_EResult AcceptResult = EOS_P2P_AcceptConnection(SocketSubsystem.GetP2PHandle(), &Options);
+			if (AcceptResult == EOS_EResult::EOS_Success)
+			{
+				UE_LOG(LogSocketSubsystemEOS, Verbose, TEXT("Accepting connection request from (%s) on socket (%s)"), *RemoteUser, UTF8_TO_TCHAR(Info->SocketId->SocketName));
+				NP_LOG(TEXT("[%s] - Accepting connection request from (%s) on socket (%s)\r\n"), GetLogPrefix(), *RemoteUser, UTF8_TO_TCHAR(Info->SocketId->SocketName));
+			}
+			else
+			{
+				UE_LOG(LogSocketSubsystemEOS, Error, TEXT("EOS_P2P_AcceptConnection from (%s) on socket (%s) failed with (%s)"), *RemoteUser, UTF8_TO_TCHAR(Info->SocketId->SocketName), ANSI_TO_TCHAR(EOS_EResult_ToString(AcceptResult)));
+				NP_LOG(TEXT("[%s] - EOS_P2P_AcceptConnection from (%s) on socket (%s) failed with (%s)\r\n"), GetLogPrefix(), *RemoteUser, UTF8_TO_TCHAR(Info->SocketId->SocketName), ANSI_TO_TCHAR(EOS_EResult_ToString(AcceptResult)));
+			}
+		}
+		else
+		{
+			UE_LOG(LogSocketSubsystemEOS, Warning, TEXT("Ignoring connection request from (%s) on socket (%s)"), *RemoteUser, UTF8_TO_TCHAR(Info->SocketId->SocketName));
+		}
+	};
+	ConnectNotifyId = EOS_P2P_AddNotifyPeerConnectionRequest(SocketSubsystem.GetP2PHandle(), &Options, ConnectNotifyCallback, ConnectNotifyCallback->GetCallbackPtr());
+
+	// Need to handle closures too
+	RegisterClosedNotification();
+#endif
+
+	bIsListening = true;
+
+	return true;
+}
+
+bool FSocketEOS::WaitForPendingConnection(bool& bHasPendingConnection, const FTimespan& WaitTime)
+{
+	/** Not supported - connectionless (UDP) only */
+	SocketSubsystem.SetLastSocketError(ESocketErrors::SE_EOPNOTSUPP);
+	return false;
+}
+
+bool FSocketEOS::HasPendingData(uint32& PendingDataSize)
+{
+	check(IsInGameThread() && "p2p does not support multithreading");
+
+	PendingDataSize = 0;
+
+#if WITH_EOS_SDK
+	EOS_P2P_GetNextReceivedPacketSizeOptions Options = { };
+	Options.ApiVersion = EOS_P2P_GETNEXTRECEIVEDPACKETSIZE_API_LATEST;
+	Options.LocalUserId = LocalAddress.GetLocalUserId();
+	uint8 Channel = LocalAddress.GetChannel();
+	Options.RequestedChannel = &Channel;
+
+	EOS_EResult Result = EOS_P2P_GetNextReceivedPacketSize(SocketSubsystem.GetP2PHandle(), &Options, &PendingDataSize);
+	if (Result == EOS_EResult::EOS_NotFound)
+	{
+		return false;
+	}
+	if (Result != EOS_EResult::EOS_Success)
+	{
+		UE_LOG(LogSocketSubsystemEOS, Warning, TEXT("Unable to check for data on address (%s) result code = (%s)"), *LocalAddress.ToString(true), ANSI_TO_TCHAR(EOS_EResult_ToString(Result)));
+
+		// @todo joeg - map EOS codes to UE4's
+		SocketSubsystem.SetLastSocketError(ESocketErrors::SE_EINVAL);
+		return false;
+	}
+
+	return true;
+#else
+	return false;
+#endif
+}
+
+FSocket* FSocketEOS::Accept(const FString& InSocketDescription)
+{
+	/** Not supported - connectionless (UDP) only */
+	SocketSubsystem.SetLastSocketError(ESocketErrors::SE_EOPNOTSUPP);
+	return nullptr;
+}
+
+FSocket* FSocketEOS::Accept(FInternetAddr& OutAddr, const FString& InSocketDescription)
+{
+	/** Not supported - connectionless (UDP) only */
+	SocketSubsystem.SetLastSocketError(ESocketErrors::SE_EOPNOTSUPP);
+	return nullptr;
+}
+
+bool FSocketEOS::SendTo(const uint8* Data, int32 Count, int32& OutBytesSent, const FInternetAddr& Destination)
+{
+	check(IsInGameThread() && "p2p does not support multithreading");
+
+	OutBytesSent = 0;
+
+	if (!Destination.IsValid())
+	{
+		UE_LOG(LogSocketSubsystemEOS, Warning, TEXT("Unable to send data, invalid destination address. DestinationAddress = (%s)"), *Destination.ToString(true));
+
+		SocketSubsystem.SetLastSocketError(ESocketErrors::SE_EADDRNOTAVAIL);
+		return false;
+	}
+
+#if WITH_EOS_SDK
+	if (Count > EOS_P2P_MAX_PACKET_SIZE)
+	{
+		UE_LOG(LogSocketSubsystemEOS, Warning, TEXT("Unable to send data, data over maximum size. Amount=[%d/%d] DestinationAddress = (%s)"), Count, EOS_P2P_MAX_PACKET_SIZE, *Destination.ToString(true));
+
+		SocketSubsystem.SetLastSocketError(ESocketErrors::SE_EMSGSIZE);
+		return false;
+	}
+
+	if (Count < 0)
+	{
+		UE_LOG(LogSocketSubsystemEOS, Warning, TEXT("Unable to send data, data invalid. Amount=[%d/%d] DestinationAddress = (%s)"), Count, EOS_P2P_MAX_PACKET_SIZE, *Destination.ToString(true));
+
+		SocketSubsystem.SetLastSocketError(ESocketErrors::SE_EINVAL);
+		return false;
+	}
+#endif 
+
+	if (Data == nullptr && Count != 0)
+	{
+		UE_LOG(LogSocketSubsystemEOS, Warning, TEXT("Unable to send data, data invalid. DestinationAddress = (%s)"), *Destination.ToString(true));
+
+		SocketSubsystem.SetLastSocketError(ESocketErrors::SE_EINVAL);
+		return false;
+	}
+
+	if (!LocalAddress.IsValid())
+	{
+		UE_LOG(LogSocketSubsystemEOS, Warning, TEXT("Unable to send data, socket was not initialized. DestinationAddress = (%s)"), *Destination.ToString(true));
+
+		SocketSubsystem.SetLastSocketError(ESocketErrors::SE_NOTINITIALISED);
+		return false;
+	}
+
+	const FInternetAddrEOS& DestinationAddress = static_cast<const FInternetAddrEOS&>(Destination);
+	if (LocalAddress == DestinationAddress)
+	{
+		UE_LOG(LogSocketSubsystemEOS, Warning, TEXT("Unable to send data, unable to send data to ourselves. DestinationAddress = (%s)"), *Destination.ToString(true));
+
+		SocketSubsystem.SetLastSocketError(ESocketErrors::SE_ECONNREFUSED);
+		return false;
+	}
+
+	// Check for sending to an address we explicitly closed
+	if (WasClosed(DestinationAddress))
+	{
+		UE_LOG(LogSocketSubsystemEOS, Warning, TEXT("Unable to send data to closed connection. DestinationAddress = (%s)"), *Destination.ToString(true));
+
+		SocketSubsystem.SetLastSocketError(ESocketErrors::SE_ECONNREFUSED);
+		return false;
+	}
+
+#if WITH_EOS_SDK
+	// Need to handle closures if we are a client and the server closes down on us
+	RegisterClosedNotification();
+
+	EOS_P2P_SocketId SocketId = { };
+	SocketId.ApiVersion = EOS_P2P_SOCKETID_API_LATEST;
+	FCStringAnsi::Strcpy(SocketId.SocketName, DestinationAddress.GetSocketName());
+
+	EOS_P2P_SendPacketOptions Options = { };
+	Options.ApiVersion = EOS_P2P_SENDPACKET_API_LATEST;
+	Options.LocalUserId = LocalAddress.GetLocalUserId();
+	Options.RemoteUserId = DestinationAddress.GetRemoteUserId();
+	Options.SocketId = &SocketId;
+	Options.bAllowDelayedDelivery = EOS_TRUE;
+	Options.Reliability = EOS_EPacketReliability::EOS_PR_UnreliableUnordered;
+	Options.Channel = DestinationAddress.GetChannel();
+	Options.DataLengthBytes = Count;
+	Options.Data = Data;
+	EOS_EResult Result = EOS_P2P_SendPacket(SocketSubsystem.GetP2PHandle(), &Options);
+	NP_LOG(TEXT("[%s] - EOS_P2P_SendPacket() to (%s) result code = (%s)\r\n"), GetLogPrefix(), *Destination.ToString(true), ANSI_TO_TCHAR(EOS_EResult_ToString(Result)));
+	if (Result != EOS_EResult::EOS_Success)
+	{
+		UE_LOG(LogSocketSubsystemEOS, Error, TEXT("Unable to send data to (%s) result code = (%s)"), *Destination.ToString(true), ANSI_TO_TCHAR(EOS_EResult_ToString(Result)));
+
+		// @todo joeg - map EOS codes to UE4's
+		SocketSubsystem.SetLastSocketError(ESocketErrors::SE_EINVAL);
+		return false;
+	}
+	OutBytesSent = Count;
+	return true;
+#else
+	return false;
+#endif
+}
+
+bool FSocketEOS::Send(const uint8* Data, int32 Count, int32& BytesSent)
+{
+	/** Not supported - connectionless (UDP) only */
+	SocketSubsystem.SetLastSocketError(ESocketErrors::SE_EOPNOTSUPP);
+	BytesSent = 0;
+	return false;
+}
+
+bool FSocketEOS::RecvFrom(uint8* Data, int32 BufferSize, int32& BytesRead, FInternetAddr& Source, ESocketReceiveFlags::Type Flags)
+{
+	check(IsInGameThread() && "p2p does not support multithreading");
+	BytesRead = 0;
+
+	if (BufferSize < 0)
+	{
+		UE_LOG(LogSocketSubsystemEOS, Error, TEXT("Unable to receive data, receiving buffer was invalid. BufferSize = (%d)"), BufferSize);
+
+		SocketSubsystem.SetLastSocketError(ESocketErrors::SE_EINVAL);
+		return false;
+	}
+
+	if (Flags != ESocketReceiveFlags::None)
+	{
+		// We do not support peaking / blocking until a packet comes
+		UE_LOG(LogSocketSubsystemEOS, Error, TEXT("Socket receive flags (%d) are not supported"), int32(Flags));
+
+		SocketSubsystem.SetLastSocketError(ESocketErrors::SE_EOPNOTSUPP);
+		return false;
+	}
+
+#if WITH_EOS_SDK
+	EOS_P2P_ReceivePacketOptions Options = { };
+	Options.ApiVersion = EOS_P2P_RECEIVEPACKET_API_LATEST;
+	Options.LocalUserId = LocalAddress.GetLocalUserId();
+	Options.MaxDataSizeBytes = BufferSize;
+	uint8 Channel = LocalAddress.GetChannel();
+	Options.RequestedChannel = &Channel;
+
+	EOS_ProductUserId RemoteUserId = nullptr;
+	EOS_P2P_SocketId SocketId;
+	
+	EOS_EResult Result = EOS_P2P_ReceivePacket(SocketSubsystem.GetP2PHandle(), &Options, &RemoteUserId, &SocketId, &Channel, Data, (uint32*)&BytesRead);
+	NP_LOG(TEXT("[%s] - EOS_P2P_ReceivePacket() for user (%s) and channel (%d) with result code = (%s)\r\n"), GetLogPrefix(), *MakeStringFromProductUserId(LocalAddress.GetLocalUserId()), Channel, ANSI_TO_TCHAR(EOS_EResult_ToString(Result)));
+	if (Result == EOS_EResult::EOS_NotFound)
+	{
+		// No data to read
+		SocketSubsystem.SetLastSocketError(ESocketErrors::SE_EWOULDBLOCK);
+		return false;
+	}
+	else if (Result != EOS_EResult::EOS_Success)
+	{
+		UE_LOG(LogSocketSubsystemEOS, Error, TEXT("Unable to receive data result code = (%s)"), ANSI_TO_TCHAR(EOS_EResult_ToString(Result)));
+
+		// @todo joeg - map EOS codes to UE4's
+		SocketSubsystem.SetLastSocketError(ESocketErrors::SE_EINVAL);
+		return false;
+	}
+
+	FInternetAddrEOS& SourceAddress = static_cast<FInternetAddrEOS&>(Source);
+	SourceAddress.SetLocalUserId(LocalAddress.GetLocalUserId());
+	SourceAddress.SetRemoteUserId(RemoteUserId);
+	SourceAddress.SetSocketName(SocketId.SocketName);
+	SourceAddress.SetChannel(Channel);
+
+	NP_LOG(TEXT("[%s] - EOS_P2P_ReceivePacket() of size (%d) from (%s)\r\n"), GetLogPrefix(), BytesRead, *SourceAddress.ToString(true));
+	return true;
+#else
+	return false;
+#endif
+}
+
+bool FSocketEOS::Recv(uint8* Data, int32 BufferSize, int32& BytesRead, ESocketReceiveFlags::Type Flags)
+{
+	BytesRead = 0;
+	SocketSubsystem.SetLastSocketError(ESocketErrors::SE_EOPNOTSUPP);
+	return false;
+}
+
+bool FSocketEOS::Wait(ESocketWaitConditions::Type Condition, FTimespan WaitTime)
+{
+	SocketSubsystem.SetLastSocketError(ESocketErrors::SE_EOPNOTSUPP);
+	return false;
+}
+
+ESocketConnectionState FSocketEOS::GetConnectionState()
+{
+	return ESocketConnectionState::SCS_NotConnected;
+}
+
+void FSocketEOS::GetAddress(FInternetAddr& OutAddr)
+{
+	OutAddr = LocalAddress;
+}
+
+bool FSocketEOS::GetPeerAddress(FInternetAddr& OutAddr)
+{
+	SocketSubsystem.SetLastSocketError(ESocketErrors::SE_EOPNOTSUPP);
+	return false;
+}
+
+bool FSocketEOS::SetNonBlocking(bool bIsNonBlocking)
+{
+	return true;
+}
+
+bool FSocketEOS::SetBroadcast(bool bAllowBroadcast)
+{
+	return true;
+}
+
+bool FSocketEOS::SetNoDelay(bool bIsNoDelay)
+{
+	return true;
+}
+
+bool FSocketEOS::JoinMulticastGroup(const FInternetAddr& GroupAddress)
+{
+	SocketSubsystem.SetLastSocketError(ESocketErrors::SE_EOPNOTSUPP);
+	return false;
+}
+
+bool FSocketEOS::LeaveMulticastGroup(const FInternetAddr& GroupAddress)
+{
+	SocketSubsystem.SetLastSocketError(ESocketErrors::SE_EOPNOTSUPP);
+	return false;
+}
+
+bool FSocketEOS::SetMulticastLoopback(bool bLoopback)
+{
+	SocketSubsystem.SetLastSocketError(ESocketErrors::SE_EOPNOTSUPP);
+	return false;
+}
+
+bool FSocketEOS::SetMulticastTtl(uint8 TimeToLive)
+{
+	SocketSubsystem.SetLastSocketError(ESocketErrors::SE_EOPNOTSUPP);
+	return false;
+}
+
+bool FSocketEOS::JoinMulticastGroup(const FInternetAddr& GroupAddress, const FInternetAddr& InterfaceAddress)
+{
+	SocketSubsystem.SetLastSocketError(ESocketErrors::SE_EOPNOTSUPP);
+	return false;
+}
+
+bool FSocketEOS::LeaveMulticastGroup(const FInternetAddr& GroupAddress, const FInternetAddr& InterfaceAddress)
+{
+	SocketSubsystem.SetLastSocketError(ESocketErrors::SE_EOPNOTSUPP);
+	return false;
+}
+
+bool FSocketEOS::SetMulticastInterface(const FInternetAddr& InterfaceAddress)
+{
+	SocketSubsystem.SetLastSocketError(ESocketErrors::SE_EOPNOTSUPP);
+	return false;
+}
+
+bool FSocketEOS::SetReuseAddr(bool bAllowReuse)
+{
+	return true;
+}
+
+bool FSocketEOS::SetLinger(bool bShouldLinger, int32 Timeout)
+{
+	return true;
+}
+
+bool FSocketEOS::SetRecvErr(bool bUseErrorQueue)
+{
+	return true;
+}
+
+bool FSocketEOS::SetSendBufferSize(int32 Size, int32& NewSize)
+{
+	return true;
+}
+
+bool FSocketEOS::SetReceiveBufferSize(int32 Size, int32& NewSize)
+{
+	return true;
+}
+
+int32 FSocketEOS::GetPortNo()
+{
+	return LocalAddress.GetChannel();
+}
+
+void FSocketEOS::SetLocalAddress(const FInternetAddrEOS& InLocalAddress)
+{
+	LocalAddress = InLocalAddress;
+}
+
+bool FSocketEOS::Close(const FInternetAddrEOS& RemoteAddress)
+{
+	check(IsInGameThread() && "p2p does not support multithreading");
+
+	if (!RemoteAddress.IsValid())
+	{
+		UE_LOG(LogSocketSubsystemEOS, Error, TEXT("Unable to close socket with remote address as it is invalid RemoteAddress = (%s)"), *RemoteAddress.ToString(true));
+		SocketSubsystem.SetLastSocketError(ESocketErrors::SE_EADDRNOTAVAIL);
+		return false;
+	}
+
+#if WITH_EOS_SDK
+	// So we don't reopen a connection by sending to it
+	ClosedRemotes.Add(RemoteAddress);
+
+	EOS_P2P_SocketId SocketId = { };
+	SocketId.ApiVersion = EOS_P2P_SOCKETID_API_LATEST;
+	FCStringAnsi::Strcpy(SocketId.SocketName, RemoteAddress.GetSocketName());
+
+	EOS_P2P_CloseConnectionOptions Options = { };
+	Options.ApiVersion = EOS_P2P_CLOSECONNECTION_API_LATEST;
+	Options.LocalUserId = LocalAddress.GetLocalUserId();
+	Options.RemoteUserId = RemoteAddress.GetRemoteUserId();
+	Options.SocketId = &SocketId;
+
+	EOS_EResult Result = EOS_P2P_CloseConnection(SocketSubsystem.GetP2PHandle(), &Options);
+	NP_LOG(TEXT("[%s] - EOS_P2P_CloseConnection() with remote address RemoteAddress (%s) result code (%s)\r\n"), GetLogPrefix(), *RemoteAddress.ToString(true), ANSI_TO_TCHAR(EOS_EResult_ToString(Result)));
+	if (Result != EOS_EResult::EOS_Success)
+	{
+		UE_LOG(LogSocketSubsystemEOS, Error, TEXT("Unable to close socket with remote address RemoteAddress (%s) due to error (%s)"), *RemoteAddress.ToString(true), ANSI_TO_TCHAR(EOS_EResult_ToString(Result)));
+
+		// @todo joeg - map EOS codes to UE4's
+		SocketSubsystem.SetLastSocketError(ESocketErrors::SE_EINVAL);
+		return false;
+	}
+	return true;
+#else
+	return false;
+#endif
+}
+
+void FSocketEOS::RegisterClosedNotification()
+{
+#if WITH_EOS_SDK
+	if (ClosedNotifyId != EOS_INVALID_NOTIFICATIONID)
+	{
+		// Already listening for these events so ignore
+		return;
+	}
+	
+	EOS_P2P_SocketId SocketId = { };
+	SocketId.ApiVersion = EOS_P2P_SOCKETID_API_LATEST;
+	FCStringAnsi::Strcpy(SocketId.SocketName, LocalAddress.GetSocketName());
+
+	EOS_P2P_AddNotifyPeerConnectionClosedOptions Options = { };
+	Options.ApiVersion = EOS_P2P_ADDNOTIFYPEERCONNECTIONCLOSED_API_LATEST;
+	Options.LocalUserId = LocalAddress.GetLocalUserId();
+	Options.SocketId = &SocketId;
+
+	ClosedNotifyCallback = new FClosedNotifyCallback();
+	ClosedNotifyCallback->CallbackLambda = [this](const EOS_P2P_OnRemoteConnectionClosedInfo* Info)
+	{
+		// Add this connection to the list of closed ones
+		FInternetAddrEOS RemoteAddress(Info->RemoteUserId, Info->SocketId->SocketName, LocalAddress.GetChannel());
+		RemoteAddress.SetLocalUserId(LocalAddress.GetLocalUserId());
+		ClosedRemotes.Add(RemoteAddress);
+		NP_LOG(TEXT("[%s] - Close connection received for remote address (%s)\r\n"), GetLogPrefix(), *RemoteAddress.ToString(true));
+	};
+	ClosedNotifyId = EOS_P2P_AddNotifyPeerConnectionClosed(SocketSubsystem.GetP2PHandle(), &Options, ClosedNotifyCallback, ClosedNotifyCallback->GetCallbackPtr());
+#endif
+}
diff --git a/Engine/Plugins/Online/SocketSubsystemEOS/Source/SocketSubsystemEOS/Private/SocketSubsystemEOS.cpp b/Engine/Plugins/Online/SocketSubsystemEOS/Source/SocketSubsystemEOS/Private/SocketSubsystemEOS.cpp
new file mode 100644
--- /dev/null
+++ b/Engine/Plugins/Online/SocketSubsystemEOS/Source/SocketSubsystemEOS/Private/SocketSubsystemEOS.cpp
@@ -0,0 +1,255 @@
+// Copyright Epic Games, Inc. All Rights Reserved.
+
+#include "SocketSubsystemEOS.h"
+#include "InternetAddrEOS.h"
+#include "SocketEOS.h"
+#include "SocketTypes.h"
+#include "Containers/Ticker.h"
+#include "Misc/ConfigCacheIni.h"
+#include "SocketSubsystemModule.h"
+#include "Modules/ModuleManager.h"
+#include "Misc/OutputDeviceRedirector.h"
+
+#if WITH_EOS_SDK
+	#include "eos_sdk.h"
+#endif
+
+FSocketSubsystemEOS::FSocketSubsystemEOS(IEOSPlatformHandlePtr InPlatformHandle, ISocketSubsystemEOSUtilsPtr InUtils)
+	: P2PHandle(nullptr)
+	, Utils(InUtils)
+	, LastSocketError(ESocketErrors::SE_NO_ERROR)
+{
+#if WITH_EOS_SDK
+	P2PHandle = EOS_Platform_GetP2PInterface(*InPlatformHandle);
+	if (P2PHandle == nullptr)
+	{
+		UE_LOG(LogSocketSubsystemEOS, Error, TEXT("FOnlineSubsystemEOS: failed to init EOS platform, couldn't get p2p handle"));
+	}
+#endif
+}
+
+FSocketSubsystemEOS::~FSocketSubsystemEOS()
+{
+	Utils = nullptr;
+}
+
+bool FSocketSubsystemEOS::Init(FString& Error)
+{
+	FSocketSubsystemModule& SocketSubsystem = FModuleManager::LoadModuleChecked<FSocketSubsystemModule>("Sockets");
+	SocketSubsystem.RegisterSocketSubsystem(EOS_SOCKETSUBSYSTEM, this, false);
+
+	return true;
+}
+
+void FSocketSubsystemEOS::Shutdown()
+{
+	// Destruct our sockets before we finish destructing, as they maintain a reference to us
+	TrackedSockets.Reset();
+
+	if (FSocketSubsystemModule* SocketSubsystem = FModuleManager::GetModulePtr<FSocketSubsystemModule>("Sockets"))
+	{
+		SocketSubsystem->UnregisterSocketSubsystem(EOS_SOCKETSUBSYSTEM);
+	}
+}
+
+FSocket* FSocketSubsystemEOS::CreateSocket(const FName& SocketTypeName, const FString& SocketDescription, const FName& /*unused*/)
+{
+	return TrackedSockets.Emplace_GetRef(MakeUnique<FSocketEOS>(*this, SocketDescription)).Get();
+}
+
+FResolveInfoCached* FSocketSubsystemEOS::CreateResolveInfoCached(TSharedPtr<FInternetAddr> Addr) const
+{
+	return nullptr;
+}
+
+void FSocketSubsystemEOS::DestroySocket(FSocket* Socket)
+{
+	for (auto It = TrackedSockets.CreateIterator(); It; ++It)
+	{
+		if (It->IsValid() && It->Get() == Socket)
+		{
+			It.RemoveCurrent();
+			return;
+		}
+	}
+}
+
+FAddressInfoResult FSocketSubsystemEOS::GetAddressInfo(const TCHAR* HostName, const TCHAR* ServiceName, EAddressInfoFlags /*unused*/, const FName /*unused*/, ESocketType /*unused*/)
+{
+	return FAddressInfoResult(HostName, ServiceName);
+}
+
+bool FSocketSubsystemEOS::RequiresChatDataBeSeparate()
+{
+	return false;
+}
+
+bool FSocketSubsystemEOS::RequiresEncryptedPackets()
+{
+	return false;
+}
+
+bool FSocketSubsystemEOS::GetHostName(FString& HostName)
+{
+	return false;
+}
+
+TSharedRef<FInternetAddr> FSocketSubsystemEOS::CreateInternetAddr()
+{
+	return MakeShared<FInternetAddrEOS>();
+}
+
+TSharedPtr<FInternetAddr> FSocketSubsystemEOS::GetAddressFromString(const FString& InString)
+{
+	bool bUnused;
+	TSharedPtr<FInternetAddrEOS> NewAddress = StaticCastSharedRef<FInternetAddrEOS>(CreateInternetAddr());
+	NewAddress->SetIp(*InString, bUnused);
+	return NewAddress;
+}
+
+bool FSocketSubsystemEOS::HasNetworkDevice()
+{
+	return true;
+}
+
+const TCHAR* FSocketSubsystemEOS::GetSocketAPIName() const
+{
+	return TEXT("p2pSocketsEOS");
+}
+
+ESocketErrors FSocketSubsystemEOS::GetLastErrorCode()
+{
+	return TranslateErrorCode(LastSocketError);
+}
+
+ESocketErrors FSocketSubsystemEOS::TranslateErrorCode(int32 Code)
+{
+	return static_cast<ESocketErrors>(Code);
+}
+
+bool FSocketSubsystemEOS::GetLocalAdapterAddresses(TArray<TSharedPtr<FInternetAddr>>& OutAddresses)
+{
+	TSharedRef<FInternetAddr> AdapterAddress = GetLocalBindAddr(nullptr, *GLog);
+	OutAddresses.Add(AdapterAddress);
+	return true;
+}
+
+TArray<TSharedRef<FInternetAddr>> FSocketSubsystemEOS::GetLocalBindAddresses()
+{
+	TArray<TSharedRef<FInternetAddr>> OutAddresses;
+	OutAddresses.Add(GetLocalBindAddr(nullptr, *GLog));
+	return OutAddresses;
+}
+
+TSharedRef<FInternetAddr> FSocketSubsystemEOS::GetLocalBindAddr(FOutputDevice& Out)
+{
+	return GetLocalBindAddr(nullptr, Out);
+}
+
+#if WITH_EOS_SDK
+EOS_HP2P FSocketSubsystemEOS::GetP2PHandle()
+{
+	check(P2PHandle != nullptr);
+	return P2PHandle;
+}
+
+EOS_ProductUserId FSocketSubsystemEOS::GetLocalUserId()
+{
+	return Utils->GetLocalUserId();
+}
+#endif
+
+TSharedRef<FInternetAddr> FSocketSubsystemEOS::GetLocalBindAddr(const UWorld* const OwningWorld, FOutputDevice& Out)
+{
+	TSharedRef<FInternetAddrEOS> BoundAddr = MakeShared<FInternetAddrEOS>();
+
+#if WITH_EOS_SDK
+	EOS_ProductUserId LocalUserId = GetLocalUserId();
+	if (LocalUserId == nullptr)
+	{
+		UE_LOG(LogSocketSubsystemEOS, Error, TEXT("No local user to send p2p packets with"));
+		return BoundAddr;
+	}
+	BoundAddr->SetLocalUserId(LocalUserId);
+#else
+	return BoundAddr;
+#endif
+
+	FString SessionId = Utils->GetSessionId();
+
+	if (SessionId.IsEmpty())
+	{
+		SessionId = FName(NAME_GameSession).ToString();
+	}
+
+	BoundAddr->SetSocketName(SessionId);
+
+	return BoundAddr;
+}
+
+bool FSocketSubsystemEOS::IsSocketWaitSupported() const
+{
+	return false;
+}
+
+void FSocketSubsystemEOS::SetLastSocketError(const ESocketErrors NewSocketError)
+{
+	LastSocketError = NewSocketError;
+}
+
+bool FSocketSubsystemEOS::BindChannel(const FInternetAddrEOS& Address)
+{
+	if (!Address.IsValid())
+	{
+		SetLastSocketError(ESocketErrors::SE_EINVAL);
+		return false;
+	}
+
+	const uint8 Channel = Address.GetChannel();
+
+	FChannelSet& ExistingBoundPorts = BoundAddresses.FindOrAdd(Address.GetSocketName());
+	if (ExistingBoundPorts.Contains(Channel))
+	{
+		SetLastSocketError(ESocketErrors::SE_EADDRINUSE);
+		return false;
+	}
+
+	ExistingBoundPorts.Add(Channel);
+	return true;
+}
+
+bool FSocketSubsystemEOS::UnbindChannel(const FInternetAddrEOS& Address)
+{
+	if (!Address.IsValid())
+	{
+		SetLastSocketError(ESocketErrors::SE_EINVAL);
+		return false;
+	}
+
+	const FString SocketName = Address.GetSocketName();
+	const uint8 Channel = Address.GetChannel();
+
+	// Find our sessions collection of ports
+	FChannelSet* ExistingBoundPorts = BoundAddresses.Find(SocketName);
+	if (!ExistingBoundPorts)
+	{
+		SetLastSocketError(ESocketErrors::SE_ENOTSOCK);
+		return false;
+	}
+
+	// Remove our port and check if we had it bound
+	if (ExistingBoundPorts->Remove(Channel) == 0)
+	{
+		SetLastSocketError(ESocketErrors::SE_ENOTSOCK);
+		return false;
+	}
+
+	// Remove any empty sets
+	if (ExistingBoundPorts->Num() == 0)
+	{
+		BoundAddresses.Remove(SocketName);
+		ExistingBoundPorts = nullptr;
+	}
+
+	return true;
+}
diff --git a/Engine/Plugins/Online/SocketSubsystemEOS/Source/SocketSubsystemEOS/Private/SocketSubsystemEOSModule.cpp b/Engine/Plugins/Online/SocketSubsystemEOS/Source/SocketSubsystemEOS/Private/SocketSubsystemEOSModule.cpp
new file mode 100644
--- /dev/null
+++ b/Engine/Plugins/Online/SocketSubsystemEOS/Source/SocketSubsystemEOS/Private/SocketSubsystemEOSModule.cpp
@@ -0,0 +1,18 @@
+// Copyright Epic Games, Inc. All Rights Reserved.
+
+#include "SocketSubsystemEOSModule.h"
+
+#include "Features/IModularFeatures.h"
+#include "Modules/ModuleManager.h"
+
+IMPLEMENT_MODULE(FSocketSubsystemEOSModule, SocketSubsystemEOS);
+
+void FSocketSubsystemEOSModule::StartupModule()
+{
+
+}
+
+void FSocketSubsystemEOSModule::ShutdownModule()
+{
+
+}
\ No newline at end of file
diff --git a/Engine/Plugins/Online/SocketSubsystemEOS/Source/SocketSubsystemEOS/Private/SocketSubsystemEOSModule.h b/Engine/Plugins/Online/SocketSubsystemEOS/Source/SocketSubsystemEOS/Private/SocketSubsystemEOSModule.h
new file mode 100644
--- /dev/null
+++ b/Engine/Plugins/Online/SocketSubsystemEOS/Source/SocketSubsystemEOS/Private/SocketSubsystemEOSModule.h
@@ -0,0 +1,18 @@
+// Copyright Epic Games, Inc. All Rights Reserved.
+
+#pragma once
+
+#include "Modules/ModuleInterface.h"
+
+class FSocketSubsystemEOSModule: public IModuleInterface
+{
+public:
+	FSocketSubsystemEOSModule() = default;
+	~FSocketSubsystemEOSModule() = default;
+
+private:
+	// ~Begin IModuleInterface
+	virtual void StartupModule() override;
+	virtual void ShutdownModule() override;
+	// ~End IModuleInterface
+};
\ No newline at end of file
diff --git a/Engine/Plugins/Online/SocketSubsystemEOS/Source/SocketSubsystemEOS/Public/InternetAddrEOS.h b/Engine/Plugins/Online/SocketSubsystemEOS/Source/SocketSubsystemEOS/Public/InternetAddrEOS.h
new file mode 100644
--- /dev/null
+++ b/Engine/Plugins/Online/SocketSubsystemEOS/Source/SocketSubsystemEOS/Public/InternetAddrEOS.h
@@ -0,0 +1,176 @@
+// Copyright Epic Games, Inc. All Rights Reserved.
+
+#pragma once
+
+#include "CoreMinimal.h"
+#include "IPAddress.h"
+
+#if WITH_EOS_SDK
+#if defined(EOS_PLATFORM_BASE_FILE_NAME)
+#include EOS_PLATFORM_BASE_FILE_NAME
+#endif
+#include "eos_common.h"
+#endif
+
+DECLARE_LOG_CATEGORY_EXTERN(LogSocketSubsystemEOS, Log, All);
+
+#define EOS_SOCKET_NAME_SIZE 33
+
+class SOCKETSUBSYSTEMEOS_API FInternetAddrEOS
+	: public FInternetAddr
+{
+public:
+	FInternetAddrEOS();
+	FInternetAddrEOS(const FString& InRemoteUserId, const FString& InSocketName, const int32 InChannel);
+#if WITH_EOS_SDK
+	FInternetAddrEOS(const EOS_ProductUserId InRemoteUserId, const FString& InSocketName, const int32 InChannel);
+#endif
+	virtual ~FInternetAddrEOS() = default;
+
+//~ Begin FInternetAddr Interface
+	virtual void SetIp(uint32 InAddr) override;
+	virtual void SetIp(const TCHAR* InAddr, bool& bIsValid) override;
+	virtual void GetIp(uint32& OutAddr) const override;
+	virtual void SetPort(int32 InPort) override;
+	virtual int32 GetPort() const override;
+	virtual void SetRawIp(const TArray<uint8>& RawAddr) override;
+	virtual TArray<uint8> GetRawIp() const override;
+	virtual void SetAnyAddress() override;
+	virtual void SetBroadcastAddress() override;
+	virtual void SetLoopbackAddress() override;
+	virtual FString ToString(bool bAppendPort) const override;
+	virtual uint32 GetTypeHash() const override;
+	virtual bool IsValid() const override;
+	virtual TSharedRef<FInternetAddr> Clone() const override;
+//~ End FInternetAddr Interface
+
+	FORCEINLINE FInternetAddrEOS& operator=(const FInternetAddrEOS& Other)
+	{
+		LocalUserId = Other.LocalUserId;
+		RemoteUserId = Other.RemoteUserId;
+		FCStringAnsi::Strcpy(SocketName, Other.SocketName);
+		Channel = Other.Channel;
+		return *this;
+	}
+	
+	FORCEINLINE friend bool operator==(const FInternetAddrEOS& A, const FInternetAddrEOS& B)
+	{
+		return A.Channel == B.Channel
+#if WITH_EOS_SDK
+			&& A.LocalUserId == B.LocalUserId
+			&& A.RemoteUserId == B.RemoteUserId
+#endif
+			&& FCStringAnsi::Stricmp(A.GetSocketName(), B.GetSocketName()) == 0;
+	}
+
+	FORCEINLINE friend bool operator!=(const FInternetAddrEOS& A, const FInternetAddrEOS& B)
+	{
+		return !(A == B);
+	}
+
+	FORCEINLINE friend uint32 GetTypeHash(const FInternetAddrEOS& Address)
+	{
+		return Address.GetTypeHash();
+	}
+
+	friend bool operator<(const FInternetAddrEOS& Left, const FInternetAddrEOS& Right)
+	{
+#if WITH_EOS_SDK
+		if (Left.GetLocalUserId() < Right.GetLocalUserId())
+		{
+			return true;
+		}
+		else if (Left.GetRemoteUserId() < Right.GetRemoteUserId())
+		{
+			return true;
+		}
+		else
+#endif
+		if (FCStringAnsi::Stricmp(Left.GetSocketName(), Right.GetSocketName()) < 0)
+		{
+			return true;
+		}
+
+		return Left.GetChannel() < Right.GetChannel();
+	}
+
+#if WITH_EOS_SDK
+	EOS_ProductUserId GetLocalUserId() const
+	{
+		return LocalUserId;
+	}
+
+	void SetLocalUserId(EOS_ProductUserId InLocalUserId)
+	{
+		LocalUserId = InLocalUserId;
+	}
+
+	void SetRemoteUserId(EOS_ProductUserId InRemoteUserId)
+	{
+		RemoteUserId = InRemoteUserId;
+	}
+
+	EOS_ProductUserId GetRemoteUserId() const
+	{
+		return RemoteUserId;
+	}
+#else
+	void* GetLocalUserId() const
+	{
+		return LocalUserId;
+	}
+
+	void SetLocalUserId(void* InLocalUserId)
+	{
+		LocalUserId = InLocalUserId;
+	}
+
+	void SetRemoteUserId(void* InRemoteUserId)
+	{
+		RemoteUserId = InRemoteUserId;
+	}
+
+	void* GetRemoteUserId() const
+	{
+		return RemoteUserId;
+	}
+#endif
+	
+	const char* GetSocketName() const
+	{
+		return SocketName;
+	}
+
+	void SetSocketName(const FString& InSocketName)
+	{
+		FCStringAnsi::Strncpy(SocketName, TCHAR_TO_UTF8(*InSocketName), EOS_SOCKET_NAME_SIZE);
+	}
+
+	void SetSocketName(const char* InSocketName)
+	{
+		FCStringAnsi::Strncpy(SocketName, InSocketName, EOS_SOCKET_NAME_SIZE);
+	}
+
+	uint8 GetChannel() const
+	{
+		return Channel;
+	}
+
+	void SetChannel(uint8 InChannel)
+	{
+		Channel = InChannel;
+	}
+
+private:
+#if WITH_EOS_SDK
+	EOS_ProductUserId LocalUserId;
+	EOS_ProductUserId RemoteUserId;
+#else
+	void* LocalUserId;
+	void* RemoteUserId;
+#endif
+	char SocketName[EOS_SOCKET_NAME_SIZE];
+	uint8 Channel;
+
+	friend class SocketSubsystemEOS;
+};
diff --git a/Engine/Plugins/Online/SocketSubsystemEOS/Source/SocketSubsystemEOS/Public/NetDriverEOS.h b/Engine/Plugins/Online/SocketSubsystemEOS/Source/SocketSubsystemEOS/Public/NetDriverEOS.h
new file mode 100644
--- /dev/null
+++ b/Engine/Plugins/Online/SocketSubsystemEOS/Source/SocketSubsystemEOS/Public/NetDriverEOS.h
@@ -0,0 +1,36 @@
+// Copyright Epic Games, Inc. All Rights Reserved.
+
+#pragma once
+
+#include "CoreMinimal.h"
+#include "IpNetDriver.h"
+#include "NetDriverEOS.generated.h"
+
+class ISocketSubsystem;
+
+UCLASS(Transient, Config=Engine)
+class SOCKETSUBSYSTEMEOS_API UNetDriverEOS
+	: public UIpNetDriver
+{
+	GENERATED_BODY()
+
+public:
+//~ Begin UNetDriver Interface
+	virtual bool IsAvailable() const override;
+	virtual bool InitBase(bool bInitAsClient, FNetworkNotify* InNotify, const FURL& URL, bool bReuseAddressAndPort, FString& Error) override;
+	virtual bool InitConnect(FNetworkNotify* InNotify, const FURL& ConnectURL, FString& Error) override;
+	virtual bool InitListen(FNetworkNotify* InNotify, FURL& LocalURL, bool bReuseAddressAndPort, FString& Error) override;
+	virtual ISocketSubsystem* GetSocketSubsystem() override;
+	virtual void Shutdown() override;
+	virtual int GetClientPort() override;
+//~ End UNetDriver Interface
+
+	UWorld* FindWorld() const;
+
+public:
+	UPROPERTY()
+	bool bIsPassthrough;
+
+	UPROPERTY(Config)
+	bool bIsUsingP2PSockets;
+};
diff --git a/Engine/Plugins/Online/SocketSubsystemEOS/Source/SocketSubsystemEOS/Public/SocketEOS.h b/Engine/Plugins/Online/SocketSubsystemEOS/Source/SocketSubsystemEOS/Public/SocketEOS.h
new file mode 100644
--- /dev/null
+++ b/Engine/Plugins/Online/SocketSubsystemEOS/Source/SocketSubsystemEOS/Public/SocketEOS.h
@@ -0,0 +1,104 @@
+// Copyright Epic Games, Inc. All Rights Reserved.
+
+#pragma once
+
+#include "CoreMinimal.h"
+#include "Sockets.h"
+#include "InternetAddrEOS.h"
+#include "EOSSharedTypes.h"
+
+class FOnlineSubsystemEOS;
+class FSocketSubsystemEOS;
+
+#define WANTS_NP_LOGGING 0
+
+#if WANTS_NP_LOGGING
+	#define NP_LOG(Msg, ...) NpLog(*FString::Printf(Msg, __VA_ARGS__))
+
+	void NpLog(const TCHAR* Msg);
+#else
+	#define NP_LOG(Msg, ...)
+#endif
+
+#if WITH_EOS_SDK
+	#include "eos_p2p_types.h"
+
+	typedef TEOSGlobalCallback<EOS_P2P_OnIncomingConnectionRequestCallback, EOS_P2P_OnIncomingConnectionRequestInfo> FConnectNotifyCallback;
+	typedef TEOSGlobalCallback<EOS_P2P_OnRemoteConnectionClosedCallback, EOS_P2P_OnRemoteConnectionClosedInfo> FClosedNotifyCallback;
+#endif
+
+class SOCKETSUBSYSTEMEOS_API FSocketEOS
+	: public FSocket
+{
+public:
+	FSocketEOS(FSocketSubsystemEOS& SocketSubsystem, const FString& InSocketDescription);
+	virtual ~FSocketEOS();
+
+	//~ Begin FSocket Interface
+	virtual bool Shutdown(ESocketShutdownMode Mode) override;
+	virtual bool Close() override;
+	virtual bool Bind(const FInternetAddr& Addr) override;
+	virtual bool Connect(const FInternetAddr& Addr) override;
+	virtual bool Listen(int32 MaxBacklog) override;
+	virtual bool WaitForPendingConnection(bool& bHasPendingConnection, const FTimespan& WaitTime) override;
+	virtual bool HasPendingData(uint32& PendingDataSize) override;
+	virtual class FSocket* Accept(const FString& InSocketDescription) override;
+	virtual class FSocket* Accept(FInternetAddr& OutAddr, const FString& InSocketDescription) override;
+	virtual bool SendTo(const uint8* Data, int32 Count, int32& BytesSent, const FInternetAddr& Destination) override;
+	virtual bool Send(const uint8* Data, int32 Count, int32& BytesSent) override;
+	virtual bool RecvFrom(uint8* Data, int32 BufferSize, int32& BytesRead, FInternetAddr& Source, ESocketReceiveFlags::Type Flags = ESocketReceiveFlags::None) override;
+	virtual bool Recv(uint8* Data, int32 BufferSize, int32& BytesRead, ESocketReceiveFlags::Type Flags = ESocketReceiveFlags::None) override;
+	virtual bool Wait(ESocketWaitConditions::Type Condition, FTimespan WaitTime) override;
+	virtual ESocketConnectionState GetConnectionState() override;
+	virtual void GetAddress(FInternetAddr& OutAddr) override;
+	virtual bool GetPeerAddress(FInternetAddr& OutAddr) override;
+	virtual bool SetNonBlocking(bool bIsNonBlocking = true) override;
+	virtual bool SetBroadcast(bool bAllowBroadcast = true) override;
+	virtual bool SetNoDelay(bool bIsNoDelay = true) override;
+	virtual bool JoinMulticastGroup(const FInternetAddr& GroupAddress) override;
+	virtual bool LeaveMulticastGroup(const FInternetAddr& GroupAddress) override;
+	virtual bool SetMulticastLoopback(bool bLoopback) override;
+	virtual bool SetMulticastTtl(uint8 TimeToLive) override;
+	virtual bool JoinMulticastGroup(const FInternetAddr& GroupAddress, const FInternetAddr& InterfaceAddress) override;
+	virtual bool LeaveMulticastGroup(const FInternetAddr& GroupAddress, const FInternetAddr& InterfaceAddress) override;
+	virtual bool SetMulticastInterface(const FInternetAddr& InterfaceAddress) override;
+	virtual bool SetReuseAddr(bool bAllowReuse = true) override;
+	virtual bool SetLinger(bool bShouldLinger = true, int32 Timeout = 0) override;
+	virtual bool SetRecvErr(bool bUseErrorQueue = true) override;
+	virtual bool SetSendBufferSize(int32 Size, int32& NewSize) override;
+	virtual bool SetReceiveBufferSize(int32 Size, int32& NewSize) override;
+	virtual int32 GetPortNo() override;
+	//~ End FSocket Interface
+
+	void SetLocalAddress(const FInternetAddrEOS& InLocalAddress);
+
+	bool Close(const FInternetAddrEOS& RemoteAddress);
+
+	bool WasClosed(const FInternetAddrEOS& RemoteAddress)
+	{
+		int32 Index = -1;
+		return ClosedRemotes.Find(RemoteAddress, Index);
+	}
+
+	void RegisterClosedNotification();
+
+private:
+	/** Reference to our subsystem */
+	FSocketSubsystemEOS& SocketSubsystem;
+
+	/** Our local address; session/port will be invalid when not bound */
+	FInternetAddrEOS LocalAddress;
+
+	/** Are we currently listening? */
+	bool bIsListening;
+
+	TArray<FInternetAddrEOS> ClosedRemotes;
+
+#if WITH_EOS_SDK
+	FConnectNotifyCallback* ConnectNotifyCallback;
+	EOS_NotificationId ConnectNotifyId;
+
+	FClosedNotifyCallback* ClosedNotifyCallback;
+	EOS_NotificationId ClosedNotifyId;
+#endif
+};
diff --git a/Engine/Plugins/Online/SocketSubsystemEOS/Source/SocketSubsystemEOS/Public/SocketSubsystemEOS.h b/Engine/Plugins/Online/SocketSubsystemEOS/Source/SocketSubsystemEOS/Public/SocketSubsystemEOS.h
new file mode 100644
--- /dev/null
+++ b/Engine/Plugins/Online/SocketSubsystemEOS/Source/SocketSubsystemEOS/Public/SocketSubsystemEOS.h
@@ -0,0 +1,95 @@
+// Copyright Epic Games, Inc. All Rights Reserved.
+
+#pragma once
+
+#include "CoreMinimal.h"
+#include "SocketSubsystem.h"
+#include "IEOSSDKManager.h"
+#include "SocketSubsystemEOSUtils.h"
+
+#if WITH_EOS_SDK
+	#if defined(EOS_PLATFORM_BASE_FILE_NAME)
+	#include EOS_PLATFORM_BASE_FILE_NAME
+	#endif
+
+	#include "eos_p2p_types.h"
+#endif
+
+class UNetConnectionEOS;
+class FSocket;
+class FResolveInfoCached;
+class FResolveInfo;
+class FInternetAddr;
+class FInternetAddrEOS;
+class FSocketEOS;
+
+typedef TSet<uint8> FChannelSet;
+
+class SOCKETSUBSYSTEMEOS_API FSocketSubsystemEOS
+	: public ISocketSubsystem
+{
+public:
+	FSocketSubsystemEOS(IEOSPlatformHandlePtr InPlatformHandle, ISocketSubsystemEOSUtilsPtr InUtils);
+	virtual ~FSocketSubsystemEOS();
+
+//~ Begin ISocketSubsystem Interface
+	virtual bool Init(FString& Error) override;
+	virtual void Shutdown() override;
+	virtual FSocket* CreateSocket(const FName& SocketType, const FString& SocketDescription, const FName& ProtocolType) override;
+	virtual FResolveInfoCached* CreateResolveInfoCached(TSharedPtr<FInternetAddr> Addr) const override;
+	virtual void DestroySocket(FSocket* Socket) override;
+	virtual FAddressInfoResult GetAddressInfo(const TCHAR* HostName, const TCHAR* ServiceName = nullptr, EAddressInfoFlags QueryFlags = EAddressInfoFlags::Default, const FName ProtocolTypeName = NAME_None, ESocketType SocketType = ESocketType::SOCKTYPE_Unknown) override;
+	virtual bool RequiresChatDataBeSeparate() override;
+	virtual bool RequiresEncryptedPackets() override;
+	virtual bool GetHostName(FString& HostName) override;
+	virtual TSharedRef<FInternetAddr> CreateInternetAddr() override;
+	virtual TSharedPtr<FInternetAddr> GetAddressFromString(const FString& InString) override;
+	virtual bool HasNetworkDevice() override;
+	virtual const TCHAR* GetSocketAPIName() const override;
+	virtual ESocketErrors GetLastErrorCode() override;
+	virtual ESocketErrors TranslateErrorCode(int32 Code) override;
+	virtual bool GetLocalAdapterAddresses(TArray<TSharedPtr<FInternetAddr>>& OutAdresses) override;
+	virtual TSharedRef<FInternetAddr> GetLocalBindAddr(FOutputDevice& Out) override;
+	virtual TArray<TSharedRef<FInternetAddr>> GetLocalBindAddresses() override;
+	TSharedRef<FInternetAddr> GetLocalBindAddr(const UWorld* const OwningWorld, FOutputDevice& Out);
+	virtual bool IsSocketWaitSupported() const override;
+//~ End ISocketSubsystem Interface
+
+	void SetLastSocketError(const ESocketErrors NewSocketError);
+
+#if WITH_EOS_SDK
+	EOS_HP2P GetP2PHandle();
+	EOS_ProductUserId GetLocalUserId();
+#endif
+
+	/**
+	 * Bind our socket name & channel and ensure no other connections are using this combination
+	 *
+	 * @param Address The address to bind to
+	 * @return True if we bound successfully, false if not and LastSocketError will be set with a reason
+	 */
+	bool BindChannel(const FInternetAddrEOS& Address);
+
+	/**
+	 * Unbind a previously bound socket name and channel.
+	 *
+	 * @param Address The Address to unbind
+	 * @return True if we released the port successfully, false if not and LastSocketError will be set with a reason
+	 */
+	bool UnbindChannel(const FInternetAddrEOS& Address);
+
+private:
+#if WITH_EOS_SDK
+	EOS_HP2P P2PHandle;
+#endif
+	ISocketSubsystemEOSUtilsPtr Utils;
+
+	/** All sockets allocated by this subsystem */
+	TArray<TUniquePtr<FSocketEOS>> TrackedSockets;
+
+	/** All bound socket addresses */
+	TMap<FString, FChannelSet> BoundAddresses;
+
+	/** The last error we received */
+	ESocketErrors LastSocketError;
+};
diff --git a/Engine/Plugins/Online/SocketSubsystemEOS/Source/SocketSubsystemEOS/Public/SocketSubsystemEOSUtils.h b/Engine/Plugins/Online/SocketSubsystemEOS/Source/SocketSubsystemEOS/Public/SocketSubsystemEOSUtils.h
new file mode 100644
--- /dev/null
+++ b/Engine/Plugins/Online/SocketSubsystemEOS/Source/SocketSubsystemEOS/Public/SocketSubsystemEOSUtils.h
@@ -0,0 +1,25 @@
+// Copyright Epic Games, Inc. All Rights Reserved.
+
+#pragma once
+
+#include "CoreMinimal.h"
+
+#if WITH_EOS_SDK
+#if defined(EOS_PLATFORM_BASE_FILE_NAME)
+#include EOS_PLATFORM_BASE_FILE_NAME
+#endif
+
+#include "eos_common.h"
+#endif
+
+class SOCKETSUBSYSTEMEOS_API ISocketSubsystemEOSUtils
+{
+public:
+	virtual ~ISocketSubsystemEOSUtils() {};
+
+#if WITH_EOS_SDK
+	virtual EOS_ProductUserId GetLocalUserId() = 0;
+#endif
+	virtual FString GetSessionId() = 0;
+};
+typedef TSharedPtr<ISocketSubsystemEOSUtils, ESPMode::ThreadSafe> ISocketSubsystemEOSUtilsPtr;
diff --git a/Engine/Plugins/Online/SocketSubsystemEOS/Source/SocketSubsystemEOS/SocketSubsystemEOS.Build.cs b/Engine/Plugins/Online/SocketSubsystemEOS/Source/SocketSubsystemEOS/SocketSubsystemEOS.Build.cs
new file mode 100644
--- /dev/null
+++ b/Engine/Plugins/Online/SocketSubsystemEOS/Source/SocketSubsystemEOS/SocketSubsystemEOS.Build.cs
@@ -0,0 +1,34 @@
+// Copyright Epic Games, Inc. All Rights Reserved.
+
+using System;
+using System.IO;
+using UnrealBuildTool;
+
+public class SocketSubsystemEOS : ModuleRules
+{
+	public SocketSubsystemEOS(ReadOnlyTargetRules Target) : base(Target)
+	{
+		Type = ModuleType.CPlusPlus;
+
+		PublicDependencyModuleNames.AddRange(
+			new string[]
+			{
+				"Engine",
+				"NetCore",
+				"Sockets",
+				"OnlineSubsystemUtils"
+			}
+		);
+
+		PrivateDependencyModuleNames.AddRange(
+			new string[]
+			{
+				"Core",
+				"CoreOnline",
+				"CoreUObject",
+				"EOSSDK",
+				"EOSShared"
+			}
+		);
+	}
+}
