Exclusion of unsupported platform for plugin

diff --git a/Engine/Plugins/Online/SocketSubsystemEOS/SocketSubsystemEOS.uplugin b/Engine/Plugins/Online/SocketSubsystemEOS/SocketSubsystemEOS.uplugin
--- a/Engine/Plugins/Online/SocketSubsystemEOS/SocketSubsystemEOS.uplugin
+++ b/Engine/Plugins/Online/SocketSubsystemEOS/SocketSubsystemEOS.uplugin
@@ -18,6 +18,7 @@
 			"PlatformDenyList": [
 				"LinuxArm64",
 				"Hololens",
+				"Stadia",
 				"TVOS"
 			]
 		}
