Refactor of platform support denial

diff --git a/Engine/Plugins/Online/SocketSubsystemEOS/SocketSubsystemEOS.uplugin b/Engine/Plugins/Online/SocketSubsystemEOS/SocketSubsystemEOS.uplugin
--- a/Engine/Plugins/Online/SocketSubsystemEOS/SocketSubsystemEOS.uplugin
+++ b/Engine/Plugins/Online/SocketSubsystemEOS/SocketSubsystemEOS.uplugin
@@ -18,7 +18,6 @@
 			"PlatformDenyList": [
 				"LinuxArm64",
 				"Hololens",
-				"Stadia",
 				"TVOS"
 			]
 		}
