Implementation of further feedback for SocketSubsystemEOS plugin and related refactoring.

diff --git a/Engine/Plugins/Online/EOSShared/Source/EOSShared/Public/EOSSharedTypes.h b/Engine/Plugins/Online/EOSShared/Source/EOSShared/Public/EOSSharedTypes.h
--- a/Engine/Plugins/Online/EOSShared/Source/EOSShared/Public/EOSSharedTypes.h
+++ b/Engine/Plugins/Online/EOSShared/Source/EOSShared/Public/EOSSharedTypes.h
@@ -16,10 +16,6 @@
 #define EOS_CONNECTION_URL_PREFIX TEXT("EOS")
 #define EOS_URL_SEPARATOR TEXT(":")
 
-#ifndef EOS_SOCKETSUBSYSTEM
-#define EOS_SOCKETSUBSYSTEM FName(TEXT("EOS"))
-#endif
-
 /** Used to store a pointer to the EOS callback object without knowing type */
 class EOSSHARED_API FCallbackBase
 {
diff --git a/Engine/Plugins/Online/Experimental/OnlineServicesEOS/Source/Private/Online/OnlineServicesEOS.cpp b/Engine/Plugins/Online/Experimental/OnlineServicesEOS/Source/Private/Online/OnlineServicesEOS.cpp
--- a/Engine/Plugins/Online/Experimental/OnlineServicesEOS/Source/Private/Online/OnlineServicesEOS.cpp
+++ b/Engine/Plugins/Online/Experimental/OnlineServicesEOS/Source/Private/Online/OnlineServicesEOS.cpp
@@ -16,102 +16,12 @@
 
 #if WITH_ENGINE
 #include "InternetAddrEOS.h"
-#include "NetDriverEOS.h"
+#include "NetDriverEOSBase.h"
+#include "SocketSubsystemEOSUtils_OnlineServicesEOS.h"
 #endif
 
 namespace UE::Online {
 
-#if WITH_ENGINE
-FSocketSubsystemEOSUtils_OnlineServicesEOS::FSocketSubsystemEOSUtils_OnlineServicesEOS(FOnlineServicesEOS* InServicesEOS)
-	: ServicesEOS(InServicesEOS)
-{
-	check(ServicesEOS);
-}
-
-FSocketSubsystemEOSUtils_OnlineServicesEOS::~FSocketSubsystemEOSUtils_OnlineServicesEOS()
-{
-	ServicesEOS = nullptr;
-}
-
-EOS_ProductUserId FSocketSubsystemEOSUtils_OnlineServicesEOS::GetLocalUserId()
-{
-	EOS_ProductUserId Result = nullptr;
-
-	if (ServicesEOS)
-	{
-		using namespace UE::Online;
-
-		IAuthPtr AuthEOS = ServicesEOS->GetAuthInterface();
-		check(AuthEOS);
-
-		FAuthGetAccountByPlatformUserId::Params AuthParams;
-		AuthParams.PlatformUserId = FPlatformMisc::GetPlatformUserForUserIndex(0);
-		TOnlineResult<FAuthGetAccountByPlatformUserId> AuthResult = AuthEOS->GetAccountByPlatformUserId(MoveTemp(AuthParams));
-		if (AuthResult.IsOk())
-		{
-			UE::Online::FAuthGetAccountByPlatformUserId::Result OkValue = AuthResult.GetOkValue();
-
-			Result = GetProductUserIdChecked(OkValue.AccountInfo->UserId);
-		}
-		else
-		{
-			UE_LOG(LogOnlineServices, Verbose, TEXT("[FSocketSubsystemEOSUtils_OnlineServicesEOS::GetLocalUserId] Unable to get account for platform user id [%s]. Error=[%s]"), *ToLogString(AuthParams.PlatformUserId), *AuthResult.GetErrorValue().GetLogString(true));
-		}
-	}
-
-	return Result;
-}
-
-FString FSocketSubsystemEOSUtils_OnlineServicesEOS::GetSessionId()
-{
-	FString Result;
-
-	if (ServicesEOS)
-	{
-		using namespace UE::Online;
-
-		IAuthPtr AuthEOS = ServicesEOS->GetAuthInterface();
-		check(AuthEOS);
-
-		FAuthGetAccountByPlatformUserId::Params AuthParams;
-		AuthParams.PlatformUserId = FPlatformMisc::GetPlatformUserForUserIndex(0);
-		TOnlineResult<FAuthGetAccountByPlatformUserId> AuthResult = AuthEOS->GetAccountByPlatformUserId(MoveTemp(AuthParams));
-		if (AuthResult.IsOk())
-		{
-			FAuthGetAccountByPlatformUserId::Result* AuthOkValue = AuthResult.TryGetOkValue();
-
-			ILobbiesPtr LobbiesEOS = ServicesEOS->GetLobbiesInterface();
-			check(LobbiesEOS);
-
-			FGetJoinedLobbies::Params LobbiesParams;
-			LobbiesParams.LocalUserId = AuthOkValue->AccountInfo->UserId;
-			TOnlineResult<FGetJoinedLobbies> LobbiesResult = LobbiesEOS->GetJoinedLobbies(MoveTemp(LobbiesParams));
-			if (LobbiesResult.IsOk())
-			{
-				FGetJoinedLobbies::Result* LobbiesOkValue = LobbiesResult.TryGetOkValue();
-
-				// TODO: Pending support in Lobbies interface
-				/*for (TSharedRef<FLobby> Lobby : LobbiesOkValue->Lobbies)
-				{
-					 if (Lobby->SessionName == NAME_GameSession)
-						Result = LobbiesEOS->GetLobbyIdString(Lobby);					 
-				}*/
-			}
-			else
-			{
-				UE_LOG(LogOnlineServices, Verbose, TEXT("[FSocketSubsystemEOSUtils_OnlineServicesEOS::GetSessionId] Unable to get joined lobbies for local user id [%s]. Error=[%s]"), *ToLogString(LobbiesParams.LocalUserId), *AuthResult.GetErrorValue().GetLogString(true));
-			}
-		}
-		else
-		{
-			UE_LOG(LogOnlineServices, Verbose, TEXT("[FSocketSubsystemEOSUtils_OnlineServicesEOS::GetSessionId] Unable to get account for platform user id [%s]. Error=[%s]"), *ToLogString(AuthParams.PlatformUserId), *AuthResult.GetErrorValue().GetLogString(true));
-		}
-	}
-
-	return Result;
-}
-#endif
-
 struct FEOSPlatformConfig
 {
 	FString ProductId;
@@ -194,20 +104,20 @@ void FOnlineServicesEOS::Initialize()
 	if (EOSPlatformHandle)
 	{
 #if WITH_ENGINE
-		SocketSubsystem = MakeShareable(new FSocketSubsystemEOS(EOSPlatformHandle, MakeShareable(new FSocketSubsystemEOSUtils_OnlineServicesEOS(this))));
-		if (SocketSubsystem)
+		SocketSubsystem = MakeShareable(new FSocketSubsystemEOS(EOSPlatformHandle, MakeShareable(new FSocketSubsystemEOSUtils_OnlineServicesEOS(*this))));
+		check(SocketSubsystem);
+
+		FString ErrorStr;
+		if (!SocketSubsystem->Init(ErrorStr))
 		{
-			FString ErrorStr;
-			if (!SocketSubsystem->Init(ErrorStr))
-			{
-				UE_LOG(LogOnlineServices, Verbose, TEXT("[FOnlineServicesEOS::Initialize] Unable to initialize Socket Subsystem. Error=[%s]"), *ErrorStr);
-			}
+			UE_LOG(LogOnlineServices, Warning, TEXT("[FOnlineServicesEOS::Initialize] Unable to initialize Socket Subsystem. Error=[%s]"), *ErrorStr);
 		}
 #endif
 	}
 	else
 	{
 		UE_LOG(LogOnlineServices, Verbose, TEXT("[FOnlineServicesEOS::Initialize] Unable to initialize Socket Subsystem. EOS Platform Handle was invalid."));
+		return;
 	}
 
 	FOnlineServicesCommon::Initialize();
@@ -227,7 +137,7 @@ TOnlineResult<FGetResolvedConnectString> FOnlineServicesEOS::GetResolvedConnectS
 			{
 #if WITH_ENGINE
  				//It should look like this: "EOS:0002aeeb5b2d4388a3752dd6d31222ec:GameNetDriver:97"
- 				FString NetDriverName = GetDefault<UNetDriverEOS>()->NetDriverName.ToString();
+ 				FString NetDriverName = GetDefault<UNetDriverEOSBase>()->NetDriverName.ToString();
  				FInternetAddrEOS TempAddr(GetProductUserIdChecked(Lobby->OwnerAccountId), NetDriverName, GetTypeHash(NetDriverName));
  				return TOnlineResult<FGetResolvedConnectString>({ TempAddr.ToString(true) });
 #else
diff --git a/Engine/Plugins/Online/Experimental/OnlineServicesEOS/Source/Private/Online/SocketSubsystemEOSUtils_OnlineServicesEOS.cpp b/Engine/Plugins/Online/Experimental/OnlineServicesEOS/Source/Private/Online/SocketSubsystemEOSUtils_OnlineServicesEOS.cpp
new file mode 100644
--- /dev/null
+++ b/Engine/Plugins/Online/Experimental/OnlineServicesEOS/Source/Private/Online/SocketSubsystemEOSUtils_OnlineServicesEOS.cpp
@@ -0,0 +1,90 @@
+// Copyright Epic Games, Inc. All Rights Reserved.
+
+#include "SocketSubsystemEOSUtils_OnlineServicesEOS.h"
+
+#include "Online/OnlineServicesEOS.h"
+#include "Online/OnlineServicesEOSTypes.h"
+#include "Online/AuthEOS.h"
+#include "Online/LobbiesEOS.h"
+#include "Online/OnlineIdEOS.h"
+
+namespace UE::Online {
+
+FSocketSubsystemEOSUtils_OnlineServicesEOS::FSocketSubsystemEOSUtils_OnlineServicesEOS(FOnlineServicesEOS& InServicesEOS)
+	: ServicesEOS(InServicesEOS)
+{
+}
+
+FSocketSubsystemEOSUtils_OnlineServicesEOS::~FSocketSubsystemEOSUtils_OnlineServicesEOS()
+{
+}
+
+EOS_ProductUserId FSocketSubsystemEOSUtils_OnlineServicesEOS::GetLocalUserId()
+{
+	EOS_ProductUserId Result = nullptr;
+
+	IAuthPtr AuthEOS = ServicesEOS.GetAuthInterface();
+	check(AuthEOS);
+
+	FAuthGetAccountByPlatformUserId::Params AuthParams;
+	AuthParams.PlatformUserId = FPlatformMisc::GetPlatformUserForUserIndex(0);
+	TOnlineResult<FAuthGetAccountByPlatformUserId> AuthResult = AuthEOS->GetAccountByPlatformUserId(MoveTemp(AuthParams));
+	if (AuthResult.IsOk())
+	{
+		UE::Online::FAuthGetAccountByPlatformUserId::Result OkValue = AuthResult.GetOkValue();
+
+		Result = GetProductUserIdChecked(OkValue.AccountInfo->UserId);
+	}
+	else
+	{
+		UE_LOG(LogOnlineServices, Verbose, TEXT("[FSocketSubsystemEOSUtils_OnlineServicesEOS::GetLocalUserId] Unable to get account for platform user id [%s]. Error=[%s]"), *ToLogString(AuthParams.PlatformUserId), *AuthResult.GetErrorValue().GetLogString(true));
+	}
+
+	return Result;
+}
+
+FString FSocketSubsystemEOSUtils_OnlineServicesEOS::GetSessionId()
+{
+	FString Result;
+
+	IAuthPtr AuthEOS = ServicesEOS.GetAuthInterface();
+	check(AuthEOS);
+
+	FAuthGetAccountByPlatformUserId::Params AuthParams;
+	AuthParams.PlatformUserId = FPlatformMisc::GetPlatformUserForUserIndex(0);
+	TOnlineResult<FAuthGetAccountByPlatformUserId> AuthResult = AuthEOS->GetAccountByPlatformUserId(MoveTemp(AuthParams));
+	if (AuthResult.IsOk())
+	{
+		FAuthGetAccountByPlatformUserId::Result* AuthOkValue = AuthResult.TryGetOkValue();
+
+		ILobbiesPtr LobbiesEOS = ServicesEOS.GetLobbiesInterface();
+		check(LobbiesEOS);
+
+		FGetJoinedLobbies::Params LobbiesParams;
+		LobbiesParams.LocalUserId = AuthOkValue->AccountInfo->UserId;
+		TOnlineResult<FGetJoinedLobbies> LobbiesResult = LobbiesEOS->GetJoinedLobbies(MoveTemp(LobbiesParams));
+		if (LobbiesResult.IsOk())
+		{
+			FGetJoinedLobbies::Result* LobbiesOkValue = LobbiesResult.TryGetOkValue();
+
+			// TODO: Pending support in Lobbies interface
+			/*for (TSharedRef<FLobby> Lobby : LobbiesOkValue->Lobbies)
+			{
+					if (Lobby->SessionName == NAME_GameSession)
+					Result = LobbiesEOS->GetLobbyIdString(Lobby);
+			}*/
+		}
+		else
+		{
+			UE_LOG(LogOnlineServices, Verbose, TEXT("[FSocketSubsystemEOSUtils_OnlineServicesEOS::GetSessionId] Unable to get joined lobbies for local user id [%s]. Error=[%s]"), *ToLogString(LobbiesParams.LocalUserId), *AuthResult.GetErrorValue().GetLogString(true));
+		}
+	}
+	else
+	{
+		UE_LOG(LogOnlineServices, Verbose, TEXT("[FSocketSubsystemEOSUtils_OnlineServicesEOS::GetSessionId] Unable to get account for platform user id [%s]. Error=[%s]"), *ToLogString(AuthParams.PlatformUserId), *AuthResult.GetErrorValue().GetLogString(true));
+	}
+
+	return Result;
+}
+
+/* UE::Online */}
\ No newline at end of file
diff --git a/Engine/Plugins/Online/Experimental/OnlineServicesEOS/Source/Private/Online/SocketSubsystemEOSUtils_OnlineServicesEOS.h b/Engine/Plugins/Online/Experimental/OnlineServicesEOS/Source/Private/Online/SocketSubsystemEOSUtils_OnlineServicesEOS.h
new file mode 100644
--- /dev/null
+++ b/Engine/Plugins/Online/Experimental/OnlineServicesEOS/Source/Private/Online/SocketSubsystemEOSUtils_OnlineServicesEOS.h
@@ -0,0 +1,24 @@
+// Copyright Epic Games, Inc. All Rights Reserved.
+
+#pragma once
+
+#include "SocketSubsystemEOS.h"
+
+namespace UE::Online {
+
+class FOnlineServicesEOS;
+
+class FSocketSubsystemEOSUtils_OnlineServicesEOS : public ISocketSubsystemEOSUtils
+{
+public:
+	FSocketSubsystemEOSUtils_OnlineServicesEOS(FOnlineServicesEOS& InServicesEOS);
+	virtual ~FSocketSubsystemEOSUtils_OnlineServicesEOS() override;
+
+	virtual EOS_ProductUserId GetLocalUserId() override;
+	virtual FString GetSessionId() override;
+
+private:
+	FOnlineServicesEOS& ServicesEOS;
+};
+
+/* UE::Online */}
\ No newline at end of file
diff --git a/Engine/Plugins/Online/Experimental/OnlineServicesEOS/Source/Public/Online/OnlineServicesEOS.h b/Engine/Plugins/Online/Experimental/OnlineServicesEOS/Source/Public/Online/OnlineServicesEOS.h
--- a/Engine/Plugins/Online/Experimental/OnlineServicesEOS/Source/Public/Online/OnlineServicesEOS.h
+++ b/Engine/Plugins/Online/Experimental/OnlineServicesEOS/Source/Public/Online/OnlineServicesEOS.h
@@ -4,35 +4,18 @@
 
 #include "Online/OnlineServicesCommon.h"
 
-#if WITH_ENGINE
-#include "SocketSubsystemEOS.h"
-#endif
-
 #if defined(EOS_PLATFORM_BASE_FILE_NAME)
 #include EOS_PLATFORM_BASE_FILE_NAME
 #endif
 #include "eos_sdk.h"
 
-using IEOSPlatformHandlePtr = TSharedPtr<class IEOSPlatformHandle>;
-
-namespace UE::Online {
-
-class FOnlineServicesEOS;
-
 #if WITH_ENGINE
-class FSocketSubsystemEOSUtils_OnlineServicesEOS : public ISocketSubsystemEOSUtils
-{
-public:
-	FSocketSubsystemEOSUtils_OnlineServicesEOS(FOnlineServicesEOS* InServicesEOS);
-	virtual ~FSocketSubsystemEOSUtils_OnlineServicesEOS() override;
+class FSocketSubsystemEOS;
+#endif
 
-	virtual EOS_ProductUserId GetLocalUserId() override;
-	virtual FString GetSessionId() override;
+using IEOSPlatformHandlePtr = TSharedPtr<class IEOSPlatformHandle>;
 
-private:
-	FOnlineServicesEOS* ServicesEOS;
-};
-#endif
+namespace UE::Online {
 
 class ONLINESERVICESEOS_API FOnlineServicesEOS : public FOnlineServicesCommon
 {
diff --git a/Engine/Plugins/Online/OnlineSubsystemEOS/Source/OnlineSubsystemEOS/Private/NetDriverEOS.h b/Engine/Plugins/Online/OnlineSubsystemEOS/Source/OnlineSubsystemEOS/Private/NetDriverEOS.h
new file mode 100644
--- /dev/null
+++ b/Engine/Plugins/Online/OnlineSubsystemEOS/Source/OnlineSubsystemEOS/Private/NetDriverEOS.h
@@ -0,0 +1,18 @@
+// Copyright Epic Games, Inc. All Rights Reserved.
+
+#pragma once
+
+#include "CoreMinimal.h"
+#include "NetDriverEOSBase.h"
+#include "NetDriverEOS.generated.h"
+
+class ISocketSubsystem;
+
+UCLASS(Transient, Config=Engine)
+class ONLINESUBSYSTEMEOS_API UNetDriverEOS
+	: public UNetDriverEOSBase
+{
+	GENERATED_BODY()
+
+public:
+};
\ No newline at end of file
diff --git a/Engine/Plugins/Online/OnlineSubsystemEOS/Source/OnlineSubsystemEOS/Private/OnlineSessionEOS.cpp b/Engine/Plugins/Online/OnlineSubsystemEOS/Source/OnlineSubsystemEOS/Private/OnlineSessionEOS.cpp
--- a/Engine/Plugins/Online/OnlineSubsystemEOS/Source/OnlineSubsystemEOS/Private/OnlineSessionEOS.cpp
+++ b/Engine/Plugins/Online/OnlineSubsystemEOS/Source/OnlineSubsystemEOS/Private/OnlineSessionEOS.cpp
@@ -8,7 +8,7 @@
 #include "UserManagerEOS.h"
 #include "OnlineSubsystemUtils.h"
 #include "OnlineAsyncTaskManager.h"
-#include "SocketSubsystem.h"
+#include "SocketSubsystemEOS.h"
 #include "NboSerializerEOS.h"
 #include "InternetAddrEOS.h"
 #include "IEOSSDKManager.h"
@@ -493,7 +493,12 @@ void FOnlineSessionEOS::Init(const FString& InBucketId)
 
 	bIsDedicatedServer = IsRunningDedicatedServer();
 	bIsUsingP2PSockets = false;
-	GConfig->GetBool(TEXT("/Script/SocketSubsystemEOS.NetDriverEOS"), TEXT("bIsUsingP2PSockets"), bIsUsingP2PSockets, GEngineIni);
+
+ 	if (!GConfig->GetBool(TEXT("/Script/OnlineSubsystemEOS.NetDriverEOS"), TEXT("bIsUsingP2PSockets"), bIsUsingP2PSockets, GEngineIni))
+	{
+		// Fallback to base location
+		GConfig->GetBool(TEXT("/Script/SocketSubsystemEOS.NetDriverEOSBase"), TEXT("bIsUsingP2PSockets"), bIsUsingP2PSockets, GEngineIni);
+	}
 }
 
 /**
diff --git a/Engine/Plugins/Online/OnlineSubsystemEOS/Source/OnlineSubsystemEOS/Private/OnlineSubsystemEOS.cpp b/Engine/Plugins/Online/OnlineSubsystemEOS/Source/OnlineSubsystemEOS/Private/OnlineSubsystemEOS.cpp
--- a/Engine/Plugins/Online/OnlineSubsystemEOS/Source/OnlineSubsystemEOS/Private/OnlineSubsystemEOS.cpp
+++ b/Engine/Plugins/Online/OnlineSubsystemEOS/Source/OnlineSubsystemEOS/Private/OnlineSubsystemEOS.cpp
@@ -13,6 +13,7 @@
 #include "EOSSettings.h"
 #include "EOSShared.h"
 #include "IEOSSDKManager.h"
+#include "SocketSubsystemEOSUtils_OnlineSubsystemEOS.h"
 
 #include "Features/IModularFeatures.h"
 #include "Misc/App.h"
@@ -150,48 +151,6 @@ struct FEOSPlatformOptions :
 
 FPlatformEOSHelpersPtr FOnlineSubsystemEOS::EOSHelpersPtr;
 
-FSocketSubsystemEOSUtils_OnlineSubsystemEOS::FSocketSubsystemEOSUtils_OnlineSubsystemEOS(FOnlineSubsystemEOS* InSubsystemEOS)
-	: SubsystemEOS(InSubsystemEOS)
-{
-	check(SubsystemEOS != nullptr);
-}
-
-FSocketSubsystemEOSUtils_OnlineSubsystemEOS::~FSocketSubsystemEOSUtils_OnlineSubsystemEOS()
-{
-	SubsystemEOS = nullptr;
-}
-
-EOS_ProductUserId FSocketSubsystemEOSUtils_OnlineSubsystemEOS::GetLocalUserId()
-{
-	EOS_ProductUserId Result = nullptr;
-
-	if (SubsystemEOS)
-	{
-		Result = SubsystemEOS->UserManager->GetLocalProductUserId();
-	}
-
-	return Result;
-}
-
-FString FSocketSubsystemEOSUtils_OnlineSubsystemEOS::GetSessionId()
-{
-	FString Result;
-
-	if (SubsystemEOS)
-	{
-		const IOnlineSessionPtr DefaultSessionInt = SubsystemEOS->GetSessionInterface();
-		if (DefaultSessionInt.IsValid())
-		{
-			if (const FNamedOnlineSession* const NamedSession = DefaultSessionInt->GetNamedSession(NAME_GameSession))
-			{
-				Result = NamedSession->GetSessionIdStr();
-			}
-		}
-	}
-
-	return Result;
-}
-
 void FOnlineSubsystemEOS::ModuleInit()
 {
 	LLM_SCOPE(ELLMTag::RealTimeCommunications);
@@ -414,9 +373,13 @@ bool FOnlineSubsystemEOS::Init()
 		return false;
 	}
 
-	SocketSubsystem = MakeShareable(new FSocketSubsystemEOS(EOSPlatformHandle, MakeShareable(new FSocketSubsystemEOSUtils_OnlineSubsystemEOS(this))));
+	SocketSubsystem = MakeShareable(new FSocketSubsystemEOS(EOSPlatformHandle, MakeShareable(new FSocketSubsystemEOSUtils_OnlineSubsystemEOS(*this))));
+	check(SocketSubsystem);
 	FString ErrorMessage;
-	SocketSubsystem->Init(ErrorMessage);
+	if (!SocketSubsystem->Init(ErrorMessage))
+	{
+		UE_LOG_ONLINE(Warning, TEXT("[FOnlineSubsystemEOS::Init] Unable to initialize Socket Subsystem. Error=[%s]"), *ErrorMessage);
+	}
 
 	// We set the product id
 	FString ArtifactName;
diff --git a/Engine/Plugins/Online/OnlineSubsystemEOS/Source/OnlineSubsystemEOS/Private/SocketSubsystemEOSUtils_OnlineSubsystemEOS.cpp b/Engine/Plugins/Online/OnlineSubsystemEOS/Source/OnlineSubsystemEOS/Private/SocketSubsystemEOSUtils_OnlineSubsystemEOS.cpp
new file mode 100644
--- /dev/null
+++ b/Engine/Plugins/Online/OnlineSubsystemEOS/Source/OnlineSubsystemEOS/Private/SocketSubsystemEOSUtils_OnlineSubsystemEOS.cpp
@@ -0,0 +1,41 @@
+// Copyright Epic Games, Inc. All Rights Reserved.
+
+#include "SocketSubsystemEOSUtils_OnlineSubsystemEOS.h"
+
+#include "OnlineSubsystemEOS.h"
+#include "UserManagerEOS.h"
+#include "OnlineSessionEOS.h"
+
+FSocketSubsystemEOSUtils_OnlineSubsystemEOS::FSocketSubsystemEOSUtils_OnlineSubsystemEOS(FOnlineSubsystemEOS& InSubsystemEOS)
+	: SubsystemEOS(InSubsystemEOS)
+{
+}
+
+FSocketSubsystemEOSUtils_OnlineSubsystemEOS::~FSocketSubsystemEOSUtils_OnlineSubsystemEOS()
+{
+}
+
+EOS_ProductUserId FSocketSubsystemEOSUtils_OnlineSubsystemEOS::GetLocalUserId()
+{
+	EOS_ProductUserId Result = nullptr;
+
+	Result = SubsystemEOS.UserManager->GetLocalProductUserId();
+
+	return Result;
+}
+
+FString FSocketSubsystemEOSUtils_OnlineSubsystemEOS::GetSessionId()
+{
+	FString Result;
+
+	const IOnlineSessionPtr DefaultSessionInt = SubsystemEOS.GetSessionInterface();
+	if (DefaultSessionInt.IsValid())
+	{
+		if (const FNamedOnlineSession* const NamedSession = DefaultSessionInt->GetNamedSession(NAME_GameSession))
+		{
+			Result = NamedSession->GetSessionIdStr();
+		}
+	}
+
+	return Result;
+}
\ No newline at end of file
diff --git a/Engine/Plugins/Online/OnlineSubsystemEOS/Source/OnlineSubsystemEOS/Private/SocketSubsystemEOSUtils_OnlineSubsystemEOS.h b/Engine/Plugins/Online/OnlineSubsystemEOS/Source/OnlineSubsystemEOS/Private/SocketSubsystemEOSUtils_OnlineSubsystemEOS.h
new file mode 100644
--- /dev/null
+++ b/Engine/Plugins/Online/OnlineSubsystemEOS/Source/OnlineSubsystemEOS/Private/SocketSubsystemEOSUtils_OnlineSubsystemEOS.h
@@ -0,0 +1,20 @@
+// Copyright Epic Games, Inc. All Rights Reserved.
+
+#pragma once 
+
+#include "SocketSubsystemEOS.h"
+
+class FOnlineSubsystemEOS;
+
+class FSocketSubsystemEOSUtils_OnlineSubsystemEOS : public ISocketSubsystemEOSUtils
+{
+public:
+	FSocketSubsystemEOSUtils_OnlineSubsystemEOS(FOnlineSubsystemEOS& InSubsystemEOS);
+	virtual ~FSocketSubsystemEOSUtils_OnlineSubsystemEOS() override;
+
+	virtual EOS_ProductUserId GetLocalUserId() override;
+	virtual FString GetSessionId() override;
+
+private:
+	FOnlineSubsystemEOS& SubsystemEOS;
+};
\ No newline at end of file
diff --git a/Engine/Plugins/Online/OnlineSubsystemEOS/Source/OnlineSubsystemEOS/Public/OnlineSubsystemEOS.h b/Engine/Plugins/Online/OnlineSubsystemEOS/Source/OnlineSubsystemEOS/Public/OnlineSubsystemEOS.h
--- a/Engine/Plugins/Online/OnlineSubsystemEOS/Source/OnlineSubsystemEOS/Public/OnlineSubsystemEOS.h
+++ b/Engine/Plugins/Online/OnlineSubsystemEOS/Source/OnlineSubsystemEOS/Public/OnlineSubsystemEOS.h
@@ -5,7 +5,6 @@
 #include "CoreMinimal.h"
 #include "OnlineSubsystemNames.h"
 #include "OnlineSubsystemImpl.h"
-#include "SocketSubsystemEOS.h"
 
 #include COMPILED_PLATFORM_HEADER(EOSHelpers.h)
 
@@ -15,6 +14,7 @@ DECLARE_STATS_GROUP(TEXT("EOS"), STATGROUP_EOS, STATCAT_Advanced);
 
 #include "eos_sdk.h"
 
+class FSocketSubsystemEOS;
 class IEOSSDKManager;
 using IEOSPlatformHandlePtr = TSharedPtr<class IEOSPlatformHandle, ESPMode::ThreadSafe>;
 
@@ -49,19 +49,6 @@ typedef TSharedPtr<class FOnlineUserCloudEOS, ESPMode::ThreadSafe> FOnlineUserCl
 
 typedef TSharedPtr<FPlatformEOSHelpers, ESPMode::ThreadSafe> FPlatformEOSHelpersPtr;
 
-class FSocketSubsystemEOSUtils_OnlineSubsystemEOS : public ISocketSubsystemEOSUtils
-{
-public:
-	FSocketSubsystemEOSUtils_OnlineSubsystemEOS(FOnlineSubsystemEOS* InSubsystemEOS);
-	virtual ~FSocketSubsystemEOSUtils_OnlineSubsystemEOS() override;
-
-	virtual EOS_ProductUserId GetLocalUserId() override;
-	virtual FString GetSessionId() override;
-
-private:
-	FOnlineSubsystemEOS* SubsystemEOS;
-};
-
 /**
  *	OnlineSubsystemEOS - Implementation of the online subsystem for EOS services
  */
diff --git a/Engine/Plugins/Online/SocketSubsystemEOS/Source/SocketSubsystemEOS/Private/NetConnectionEOS.cpp b/Engine/Plugins/Online/SocketSubsystemEOS/Source/SocketSubsystemEOS/Private/NetConnectionEOS.cpp
--- a/Engine/Plugins/Online/SocketSubsystemEOS/Source/SocketSubsystemEOS/Private/NetConnectionEOS.cpp
+++ b/Engine/Plugins/Online/SocketSubsystemEOS/Source/SocketSubsystemEOS/Private/NetConnectionEOS.cpp
@@ -1,7 +1,7 @@
 // Copyright Epic Games, Inc. All Rights Reserved.
 
 #include "NetConnectionEOS.h"
-#include "NetDriverEOS.h"
+#include "NetDriverEOSBase.h"
 #include "InternetAddrEOS.h"
 #include "SocketEOS.h"
 
@@ -14,7 +14,7 @@ UNetConnectionEOS::UNetConnectionEOS(const FObjectInitializer& ObjectInitializer
 
 void UNetConnectionEOS::InitLocalConnection(UNetDriver* InDriver, FSocket* InSocket, const FURL& InURL, EConnectionState InState, int32 InMaxPacket, int32 InPacketOverhead)
 {
-	bIsPassthrough = !static_cast<UNetDriverEOS*>(InDriver)->bIsUsingP2PSockets || !InURL.Host.StartsWith(EOS_CONNECTION_URL_PREFIX, ESearchCase::IgnoreCase);
+	bIsPassthrough = !static_cast<UNetDriverEOSBase*>(InDriver)->bIsUsingP2PSockets || !InURL.Host.StartsWith(EOS_CONNECTION_URL_PREFIX, ESearchCase::IgnoreCase);
 	bHasP2PSession = !bIsPassthrough;
 	  
 	if (bHasP2PSession)
@@ -33,7 +33,7 @@ void UNetConnectionEOS::InitLocalConnection(UNetDriver* InDriver, FSocket* InSoc
 
 void UNetConnectionEOS::InitRemoteConnection(UNetDriver* InDriver, FSocket* InSocket, const FURL& InURL, const FInternetAddr& InRemoteAddr, EConnectionState InState, int32 InMaxPacket, int32 InPacketOverhead)
 {
-	bIsPassthrough = static_cast<UNetDriverEOS*>(InDriver)->bIsPassthrough;
+	bIsPassthrough = static_cast<UNetDriverEOSBase*>(InDriver)->bIsPassthrough;
 	bHasP2PSession = !bIsPassthrough;
 
 	if (bHasP2PSession)
diff --git a/Engine/Plugins/Online/SocketSubsystemEOS/Source/SocketSubsystemEOS/Private/NetDriverEOS.cpp b/Engine/Plugins/Online/SocketSubsystemEOS/Source/SocketSubsystemEOS/Private/NetDriverEOS.cpp
deleted file mode 100644
diff --git a/Engine/Plugins/Online/SocketSubsystemEOS/Source/SocketSubsystemEOS/Private/NetDriverEOSBase.cpp b/Engine/Plugins/Online/SocketSubsystemEOS/Source/SocketSubsystemEOS/Private/NetDriverEOSBase.cpp
new file mode 100644
--- /dev/null
+++ b/Engine/Plugins/Online/SocketSubsystemEOS/Source/SocketSubsystemEOS/Private/NetDriverEOSBase.cpp
@@ -0,0 +1,226 @@
+// Copyright Epic Games, Inc. All Rights Reserved.
+
+#include "NetDriverEOSBase.h"
+#include "NetConnectionEOS.h"
+#include "SocketEOS.h"
+#include "SocketSubsystemEOS.h"
+#include "Misc/EngineVersionComparison.h"
+#include "EOSSharedTypes.h"
+#include "Engine/Engine.h"
+
+bool UNetDriverEOSBase::IsAvailable() const
+{
+	// Use passthrough sockets if we are a dedicated server
+	if (IsRunningDedicatedServer())
+	{
+		return false;
+	}
+
+	if (ISocketSubsystem* SocketSubsystem = ISocketSubsystem::Get(EOS_SOCKETSUBSYSTEM))
+	{
+		return true;
+	}
+
+	return false;
+}
+
+bool UNetDriverEOSBase::InitBase(bool bInitAsClient, FNetworkNotify* InNotify, const FURL& URL, bool bReuseAddressAndPort, FString& Error)
+{
+	if (bIsPassthrough)
+	{
+		UE_LOG(LogTemp, Verbose, TEXT("Running as pass-through"));
+		return Super::InitBase(bInitAsClient, InNotify, URL, bReuseAddressAndPort, Error);
+	}
+
+	if (!UNetDriver::InitBase(bInitAsClient, InNotify, URL, bReuseAddressAndPort, Error))
+	{
+		UE_LOG(LogTemp, Warning, TEXT("Failed to init driver base"));
+		return false;
+	}
+
+	FSocketSubsystemEOS* const SocketSubsystem = static_cast<FSocketSubsystemEOS*>(GetSocketSubsystem());
+	if (!SocketSubsystem)
+	{
+		UE_LOG(LogTemp, Warning, TEXT("Could not get socket subsystem"));
+		return false;
+	}
+
+	// We don't care if our world is null, everything we uses handles it fine
+	const UWorld* const MyWorld = FindWorld();
+
+	// Get our local address (proves we're logged in)
+	TSharedRef<FInternetAddr> LocalAddress = SocketSubsystem->GetLocalBindAddr(MyWorld, *GLog);
+	if (!LocalAddress->IsValid())
+	{
+		// Not logged in?
+		Error = TEXT("Could not bind local address");
+		UE_LOG(LogTemp, Warning, TEXT("Could not bind local address"));
+		return false;
+	}
+
+	// Create our socket
+	SetSocketAndLocalAddress(SocketSubsystem->CreateSocket(NAME_DGram, TEXT("UE4"), NAME_None));
+	if (GetSocket() == nullptr)
+	{
+		UE_LOG(LogTemp, Warning, TEXT("Could not create socket"));
+		return false;
+	}
+
+	// Store our local address and set our port
+	TSharedRef<FInternetAddrEOS> EOSLocalAddress = StaticCastSharedRef<FInternetAddrEOS>(LocalAddress);
+	// Because some platforms remap ports, we will use the ID of the name of the net driver to be our channel
+	EOSLocalAddress->SetChannel(GetTypeHash(NetDriverName.ToString()));
+	// Set our net driver name so we don't accept connections across net driver types
+	EOSLocalAddress->SetSocketName(NetDriverName.ToString());
+
+	static_cast<FSocketEOS*>(GetSocket())->SetLocalAddress(*EOSLocalAddress);
+
+	LocalAddr = LocalAddress;
+
+	return true;
+}
+
+bool UNetDriverEOSBase::InitConnect(FNetworkNotify* InNotify, const FURL& ConnectURL, FString& Error)
+{
+	if (!bIsUsingP2PSockets || !IsAvailable() || !ConnectURL.Host.StartsWith(EOS_CONNECTION_URL_PREFIX, ESearchCase::IgnoreCase))
+	{
+		UE_LOG(LogTemp, Verbose, TEXT("Connecting using IPNetDriver passthrough. ConnectUrl = (%s)"), *ConnectURL.ToString());
+
+		bIsPassthrough = true;
+		return Super::InitConnect(InNotify, ConnectURL, Error);
+	}
+
+	bool bIsValid = false;
+	TSharedRef<FInternetAddrEOS> RemoteHost = MakeShared<FInternetAddrEOS>();
+	RemoteHost->SetIp(*ConnectURL.Host, bIsValid);
+	if (!bIsValid || ConnectURL.Port < 0)
+	{
+		Error = TEXT("Invalid remote address");
+		UE_LOG(LogTemp, Warning, TEXT("Invalid Remote Address. ConnectUrl = (%s)"), *ConnectURL.ToString());
+		return false;
+	}
+
+	UE_LOG(LogTemp, Verbose, TEXT("Connecting using EOSNetDriver. ConnectUrl = (%s)"), *ConnectURL.ToString());
+
+	if (!InitBase(true, InNotify, ConnectURL, false, Error))
+	{
+		return false;
+	}
+
+	// Set the address to what was parsed (us + remote)
+	LocalAddr = RemoteHost;
+
+	// Reference to our newly created socket
+	FSocket* CurSocket = GetSocket();
+
+	// Bind our local port
+	FSocketSubsystemEOS* const SocketSubsystem = static_cast<FSocketSubsystemEOS*>(GetSocketSubsystem());
+	check(SocketSubsystem);
+	if (!SocketSubsystem->BindNextPort(CurSocket, *LocalAddr, MaxPortCountToTry + 1, 1))
+	{
+		// Failure
+		Error = TEXT("Could not bind local port");
+		UE_LOG(LogTemp, Warning, TEXT("Could not bind local port in %d attempts"), MaxPortCountToTry);
+		return false;
+	}
+
+	// Create an unreal connection to the server
+	UNetConnectionEOS* Connection = NewObject<UNetConnectionEOS>(NetConnectionClass);
+	check(Connection);
+
+	// Set it as the server connection before anything else so everything knows this is a client
+	ServerConnection = Connection;
+	Connection->InitLocalConnection(this, CurSocket, ConnectURL, USOCK_Pending);
+
+	CreateInitialClientChannels();
+
+	return true;
+}
+
+bool UNetDriverEOSBase::InitListen(FNetworkNotify* InNotify, FURL& LocalURL, bool bReuseAddressAndPort, FString& Error)
+{
+	if (!bIsUsingP2PSockets || !IsAvailable() || LocalURL.HasOption(TEXT("bIsLanMatch")) || LocalURL.HasOption(TEXT("bUseIPSockets")))
+	{
+		UE_LOG(LogTemp, Verbose, TEXT("Init as IPNetDriver listen server. LocalURL = (%s)"), *LocalURL.ToString());
+
+		bIsPassthrough = true;
+		return Super::InitListen(InNotify, LocalURL, bReuseAddressAndPort, Error);
+	}
+
+	UE_LOG(LogTemp, Verbose, TEXT("Init as EOSNetDriver listen server. LocalURL = (%s)"), *LocalURL.ToString());
+
+	if (!InitBase(false, InNotify, LocalURL, bReuseAddressAndPort, Error))
+	{
+		return false;
+	}
+
+	// Bind our specified port if provided
+	FSocket* CurSocket = GetSocket();
+	if (!CurSocket->Listen(0))
+	{
+		Error = TEXT("Could not listen");
+		UE_LOG(LogTemp, Warning, TEXT("Could not listen on socket"));
+		return false;
+	}
+
+	InitConnectionlessHandler();
+
+	UE_LOG(LogTemp, Verbose, TEXT("Initialized as an EOSP2P listen server"));
+	return true;
+}
+
+ISocketSubsystem* UNetDriverEOSBase::GetSocketSubsystem()
+{
+	return ISocketSubsystem::Get(bIsPassthrough ? PLATFORM_SOCKETSUBSYSTEM : EOS_SOCKETSUBSYSTEM);
+}
+
+void UNetDriverEOSBase::Shutdown()
+{
+	UE_LOG(LogTemp, Verbose, TEXT("Shutting down NetDriver"));
+
+	Super::Shutdown();
+
+	// Kill our P2P sessions now, instead of when garbage collection kicks in later
+	if (!bIsPassthrough)
+	{
+		if (UNetConnectionEOS* const EOSServerConnection = Cast<UNetConnectionEOS>(ServerConnection))
+		{
+			EOSServerConnection->DestroyEOSConnection();
+		}
+		for (UNetConnection* Client : ClientConnections)
+		{
+			if (UNetConnectionEOS* const EOSClient = Cast<UNetConnectionEOS>(Client))
+			{
+				EOSClient->DestroyEOSConnection();
+			}
+		}
+	}
+}
+
+int UNetDriverEOSBase::GetClientPort()
+{
+	if (bIsPassthrough)
+	{
+		return Super::GetClientPort();
+	}
+
+	// Starting range of dynamic/private/ephemeral ports
+	return 49152;
+}
+
+UWorld* UNetDriverEOSBase::FindWorld() const
+{
+	UWorld* MyWorld = GetWorld();
+
+	// If we don't have a world, we may be a pending net driver
+	if (!MyWorld && GEngine)
+	{
+		if (FWorldContext* WorldContext = GEngine->GetWorldContextFromPendingNetGameNetDriver(this))
+		{
+			MyWorld = WorldContext->World();
+		}
+	}
+
+	return MyWorld;
+}
+
diff --git a/Engine/Plugins/Online/SocketSubsystemEOS/Source/SocketSubsystemEOS/Public/NetDriverEOS.h b/Engine/Plugins/Online/SocketSubsystemEOS/Source/SocketSubsystemEOS/Public/NetDriverEOS.h
deleted file mode 100644
diff --git a/Engine/Plugins/Online/SocketSubsystemEOS/Source/SocketSubsystemEOS/Public/NetDriverEOSBase.h b/Engine/Plugins/Online/SocketSubsystemEOS/Source/SocketSubsystemEOS/Public/NetDriverEOSBase.h
new file mode 100644
--- /dev/null
+++ b/Engine/Plugins/Online/SocketSubsystemEOS/Source/SocketSubsystemEOS/Public/NetDriverEOSBase.h
@@ -0,0 +1,36 @@
+// Copyright Epic Games, Inc. All Rights Reserved.
+
+#pragma once
+
+#include "CoreMinimal.h"
+#include "IpNetDriver.h"
+#include "NetDriverEOSBase.generated.h"
+
+class ISocketSubsystem;
+
+UCLASS(Transient, Config=Engine)
+class SOCKETSUBSYSTEMEOS_API UNetDriverEOSBase
+	: public UIpNetDriver
+{
+	GENERATED_BODY()
+
+public:
+//~ Begin UNetDriver Interface
+	virtual bool IsAvailable() const override;
+	virtual bool InitBase(bool bInitAsClient, FNetworkNotify* InNotify, const FURL& URL, bool bReuseAddressAndPort, FString& Error) override;
+	virtual bool InitConnect(FNetworkNotify* InNotify, const FURL& ConnectURL, FString& Error) override;
+	virtual bool InitListen(FNetworkNotify* InNotify, FURL& LocalURL, bool bReuseAddressAndPort, FString& Error) override;
+	virtual ISocketSubsystem* GetSocketSubsystem() override;
+	virtual void Shutdown() override;
+	virtual int GetClientPort() override;
+//~ End UNetDriver Interface
+
+	UWorld* FindWorld() const;
+
+public:
+	UPROPERTY()
+	bool bIsPassthrough;
+
+	UPROPERTY(Config)
+	bool bIsUsingP2PSockets;
+};
diff --git a/Engine/Plugins/Online/SocketSubsystemEOS/Source/SocketSubsystemEOS/Public/SocketSubsystemEOS.h b/Engine/Plugins/Online/SocketSubsystemEOS/Source/SocketSubsystemEOS/Public/SocketSubsystemEOS.h
--- a/Engine/Plugins/Online/SocketSubsystemEOS/Source/SocketSubsystemEOS/Public/SocketSubsystemEOS.h
+++ b/Engine/Plugins/Online/SocketSubsystemEOS/Source/SocketSubsystemEOS/Public/SocketSubsystemEOS.h
@@ -15,6 +15,10 @@
 	#include "eos_p2p_types.h"
 #endif
 
+#ifndef EOS_SOCKETSUBSYSTEM
+#define EOS_SOCKETSUBSYSTEM FName(TEXT("EOS"))
+#endif
+
 class UNetConnectionEOS;
 class FSocket;
 class FResolveInfoCached;
diff --git a/Engine/Plugins/Online/SocketSubsystemEOS/Source/SocketSubsystemEOS/SocketSubsystemEOS.Build.cs b/Engine/Plugins/Online/SocketSubsystemEOS/Source/SocketSubsystemEOS/SocketSubsystemEOS.Build.cs
--- a/Engine/Plugins/Online/SocketSubsystemEOS/Source/SocketSubsystemEOS/SocketSubsystemEOS.Build.cs
+++ b/Engine/Plugins/Online/SocketSubsystemEOS/Source/SocketSubsystemEOS/SocketSubsystemEOS.Build.cs
@@ -13,7 +13,9 @@ public class SocketSubsystemEOS : ModuleRules
 		PublicDependencyModuleNames.AddRange(
 			new string[]
 			{
+				"Core",
 				"Engine",
+				"EOSShared",
 				"NetCore",
 				"Sockets",
 				"OnlineSubsystemUtils"
@@ -23,11 +25,9 @@ public class SocketSubsystemEOS : ModuleRules
 		PrivateDependencyModuleNames.AddRange(
 			new string[]
 			{
-				"Core",
 				"CoreOnline",
 				"CoreUObject",
-				"EOSSDK",
-				"EOSShared"
+				"EOSSDK"
 			}
 		);
 	}
