Refactor to use PlatformAllowList instead of PlatformDenyList in SocketSubsystemEOS, solving issues when using Cook on the fly

diff --git a/Engine/Plugins/Online/SocketSubsystemEOS/SocketSubsystemEOS.uplugin b/Engine/Plugins/Online/SocketSubsystemEOS/SocketSubsystemEOS.uplugin
--- a/Engine/Plugins/Online/SocketSubsystemEOS/SocketSubsystemEOS.uplugin
+++ b/Engine/Plugins/Online/SocketSubsystemEOS/SocketSubsystemEOS.uplugin
@@ -15,10 +15,10 @@
 			"Name": "SocketSubsystemEOS",
 			"Type": "RuntimeNoCommandlet",
 			"LoadingPhase": "PostConfigInit",
-			"PlatformDenyList": [
-				"LinuxArm64",
-				"Hololens",
-				"TVOS"
+			"PlatformAllowList": [
+				"Win64",
+				"Mac",
+				"Android"
 			]
 		}
 	],
