Added logic to pick the right SocketSubsystemEOS when using EOS P2P sockets in PIE mode.

diff --git a/Engine/Plugins/Online/Experimental/OnlineServicesEOS/Source/Private/Online/SocketSubsystemEOSUtils_OnlineServicesEOS.cpp b/Engine/Plugins/Online/Experimental/OnlineServicesEOS/Source/Private/Online/SocketSubsystemEOSUtils_OnlineServicesEOS.cpp
--- a/Engine/Plugins/Online/Experimental/OnlineServicesEOS/Source/Private/Online/SocketSubsystemEOSUtils_OnlineServicesEOS.cpp
+++ b/Engine/Plugins/Online/Experimental/OnlineServicesEOS/Source/Private/Online/SocketSubsystemEOSUtils_OnlineServicesEOS.cpp
@@ -88,6 +88,12 @@ FString FSocketSubsystemEOSUtils_OnlineServicesEOS::GetSessionId()
 	return Result;
 }
 
+FName FSocketSubsystemEOSUtils_OnlineServicesEOS::GetSubsystemInstanceName()
+{
+	// TODO: OnlineServices still doesn't have functionality that matches GetInstanceName
+	return FName();
+}
+
 /* UE::Online */}
 
 #endif // WITH_ENGINE
diff --git a/Engine/Plugins/Online/Experimental/OnlineServicesEOS/Source/Private/Online/SocketSubsystemEOSUtils_OnlineServicesEOS.h b/Engine/Plugins/Online/Experimental/OnlineServicesEOS/Source/Private/Online/SocketSubsystemEOSUtils_OnlineServicesEOS.h
--- a/Engine/Plugins/Online/Experimental/OnlineServicesEOS/Source/Private/Online/SocketSubsystemEOSUtils_OnlineServicesEOS.h
+++ b/Engine/Plugins/Online/Experimental/OnlineServicesEOS/Source/Private/Online/SocketSubsystemEOSUtils_OnlineServicesEOS.h
@@ -18,6 +18,7 @@ public:
 
 	virtual EOS_ProductUserId GetLocalUserId() override;
 	virtual FString GetSessionId() override;
+	virtual FName GetSubsystemInstanceName() override;
 
 private:
 	FOnlineServicesEOS& ServicesEOS;
diff --git a/Engine/Plugins/Online/OnlineSubsystemEOS/Source/OnlineSubsystemEOS/Private/SocketSubsystemEOSUtils_OnlineSubsystemEOS.cpp b/Engine/Plugins/Online/OnlineSubsystemEOS/Source/OnlineSubsystemEOS/Private/SocketSubsystemEOSUtils_OnlineSubsystemEOS.cpp
--- a/Engine/Plugins/Online/OnlineSubsystemEOS/Source/OnlineSubsystemEOS/Private/SocketSubsystemEOSUtils_OnlineSubsystemEOS.cpp
+++ b/Engine/Plugins/Online/OnlineSubsystemEOS/Source/OnlineSubsystemEOS/Private/SocketSubsystemEOSUtils_OnlineSubsystemEOS.cpp
@@ -38,4 +38,9 @@ FString FSocketSubsystemEOSUtils_OnlineSubsystemEOS::GetSessionId()
 	}
 
 	return Result;
-}
\ No newline at end of file
+}
+
+FName FSocketSubsystemEOSUtils_OnlineSubsystemEOS::GetSubsystemInstanceName()
+{
+	return SubsystemEOS.GetInstanceName();
+}
diff --git a/Engine/Plugins/Online/OnlineSubsystemEOS/Source/OnlineSubsystemEOS/Private/SocketSubsystemEOSUtils_OnlineSubsystemEOS.h b/Engine/Plugins/Online/OnlineSubsystemEOS/Source/OnlineSubsystemEOS/Private/SocketSubsystemEOSUtils_OnlineSubsystemEOS.h
--- a/Engine/Plugins/Online/OnlineSubsystemEOS/Source/OnlineSubsystemEOS/Private/SocketSubsystemEOSUtils_OnlineSubsystemEOS.h
+++ b/Engine/Plugins/Online/OnlineSubsystemEOS/Source/OnlineSubsystemEOS/Private/SocketSubsystemEOSUtils_OnlineSubsystemEOS.h
@@ -14,6 +14,7 @@ public:
 
 	virtual EOS_ProductUserId GetLocalUserId() override;
 	virtual FString GetSessionId() override;
+	virtual FName GetSubsystemInstanceName() override;
 
 private:
 	FOnlineSubsystemEOS& SubsystemEOS;
diff --git a/Engine/Plugins/Online/SocketSubsystemEOS/Source/SocketSubsystemEOS/Private/NetDriverEOSBase.cpp b/Engine/Plugins/Online/SocketSubsystemEOS/Source/SocketSubsystemEOS/Private/NetDriverEOSBase.cpp
--- a/Engine/Plugins/Online/SocketSubsystemEOS/Source/SocketSubsystemEOS/Private/NetDriverEOSBase.cpp
+++ b/Engine/Plugins/Online/SocketSubsystemEOS/Source/SocketSubsystemEOS/Private/NetDriverEOSBase.cpp
@@ -7,6 +7,7 @@
 #include "Misc/EngineVersionComparison.h"
 #include "EOSSharedTypes.h"
 #include "Engine/Engine.h"
+#include "Misc/ConfigCacheIni.h"
 
 bool UNetDriverEOSBase::IsAvailable() const
 {
@@ -171,7 +172,29 @@ bool UNetDriverEOSBase::InitListen(FNetworkNotify* InNotify, FURL& LocalURL, boo
 
 ISocketSubsystem* UNetDriverEOSBase::GetSocketSubsystem()
 {
-	return ISocketSubsystem::Get(bIsPassthrough ? PLATFORM_SOCKETSUBSYSTEM : EOS_SOCKETSUBSYSTEM);
+	if (bIsPassthrough)
+	{
+		return ISocketSubsystem::Get(PLATFORM_SOCKETSUBSYSTEM);
+	}
+	else
+	{
+		FSocketSubsystemEOS* DefaultSocketSubsystem = static_cast<FSocketSubsystemEOS*>(ISocketSubsystem::Get(EOS_SOCKETSUBSYSTEM));
+
+		bool bUseOnlineServicesV2 = false;
+		GConfig->GetBool(TEXT("/Script/Engine.OnlineEngineInterface"), TEXT("bUseOnlineServicesV2"), bUseOnlineServicesV2, GEngineIni);
+
+		if (bUseOnlineServicesV2)
+		{
+			// V2 still doesn't have the logic to process several instances, so we'll return the default one
+			return DefaultSocketSubsystem;
+		}
+		else
+		{
+			// We may have several instances of FSocketSubsystemEOS if running in PIE mode, so we need to pick the right one for the current world
+			UWorld* CurrentWorld = FindWorld();
+			return DefaultSocketSubsystem->GetSocketSubsystemForWorld(CurrentWorld);
+		}
+	}
 }
 
 void UNetDriverEOSBase::Shutdown()
diff --git a/Engine/Plugins/Online/SocketSubsystemEOS/Source/SocketSubsystemEOS/Private/SocketSubsystemEOS.cpp b/Engine/Plugins/Online/SocketSubsystemEOS/Source/SocketSubsystemEOS/Private/SocketSubsystemEOS.cpp
--- a/Engine/Plugins/Online/SocketSubsystemEOS/Source/SocketSubsystemEOS/Private/SocketSubsystemEOS.cpp
+++ b/Engine/Plugins/Online/SocketSubsystemEOS/Source/SocketSubsystemEOS/Private/SocketSubsystemEOS.cpp
@@ -9,11 +9,15 @@
 #include "SocketSubsystemModule.h"
 #include "Modules/ModuleManager.h"
 #include "Misc/OutputDeviceRedirector.h"
+#include "OnlineSubsystemUtils.h"
 
 #if WITH_EOS_SDK
 	#include "eos_sdk.h"
 #endif
 
+TArray<FSocketSubsystemEOS*> FSocketSubsystemEOS::SocketSubsystemEOSInstances;
+TMap<UWorld*, FSocketSubsystemEOS*> FSocketSubsystemEOS::SocketSubsystemEOSPerWorldMap;
+
 FSocketSubsystemEOS::FSocketSubsystemEOS(IEOSPlatformHandlePtr InPlatformHandle, ISocketSubsystemEOSUtilsPtr InUtils)
 	: P2PHandle(nullptr)
 	, Utils(InUtils)
@@ -33,8 +37,34 @@ FSocketSubsystemEOS::~FSocketSubsystemEOS()
 	Utils = nullptr;
 }
 
+FSocketSubsystemEOS* FSocketSubsystemEOS::GetSocketSubsystemForWorld(UWorld* InWorld)
+{
+	FSocketSubsystemEOS** Result = SocketSubsystemEOSPerWorldMap.Find(InWorld);
+
+	if(!Result)
+	{
+		for (FSocketSubsystemEOS* SocketSubsystem : SocketSubsystemEOSInstances)
+		{
+			const UWorld* NewWorld = GetWorldForOnline(SocketSubsystem->Utils->GetSubsystemInstanceName());
+
+			if (NewWorld == InWorld)
+			{
+				SocketSubsystemEOSPerWorldMap.Add(InWorld, SocketSubsystem);
+
+				Result = &SocketSubsystem;
+
+				break;
+			}
+		}
+	}
+
+	return Result ? *Result : nullptr;
+}
+
 bool FSocketSubsystemEOS::Init(FString& Error)
 {
+	SocketSubsystemEOSInstances.Add(this);
+
 	FSocketSubsystemModule& SocketSubsystem = FModuleManager::LoadModuleChecked<FSocketSubsystemModule>("Sockets");
 	SocketSubsystem.RegisterSocketSubsystem(EOS_SOCKETSUBSYSTEM, this, false);
 
@@ -43,6 +73,8 @@ bool FSocketSubsystemEOS::Init(FString& Error)
 
 void FSocketSubsystemEOS::Shutdown()
 {
+	RemoveFromStaticContainers();
+
 	// Destruct our sockets before we finish destructing, as they maintain a reference to us
 	TrackedSockets.Reset();
 
@@ -52,6 +84,19 @@ void FSocketSubsystemEOS::Shutdown()
 	}
 }
 
+void FSocketSubsystemEOS::RemoveFromStaticContainers()
+{
+	for (TMap<UWorld*, FSocketSubsystemEOS*>::TIterator Iter(SocketSubsystemEOSPerWorldMap); Iter; ++Iter)
+	{
+		if (Iter.Value() == this)
+		{
+			Iter.RemoveCurrent();
+		}
+	}
+
+	SocketSubsystemEOSInstances.Remove(this);
+}
+
 FSocket* FSocketSubsystemEOS::CreateSocket(const FName& SocketTypeName, const FString& SocketDescription, const FName& /*unused*/)
 {
 	return TrackedSockets.Emplace_GetRef(MakeUnique<FSocketEOS>(*this, SocketDescription)).Get();
diff --git a/Engine/Plugins/Online/SocketSubsystemEOS/Source/SocketSubsystemEOS/Public/SocketSubsystemEOS.h b/Engine/Plugins/Online/SocketSubsystemEOS/Source/SocketSubsystemEOS/Public/SocketSubsystemEOS.h
--- a/Engine/Plugins/Online/SocketSubsystemEOS/Source/SocketSubsystemEOS/Public/SocketSubsystemEOS.h
+++ b/Engine/Plugins/Online/SocketSubsystemEOS/Source/SocketSubsystemEOS/Public/SocketSubsystemEOS.h
@@ -82,6 +82,19 @@ public:
 	 */
 	bool UnbindChannel(const FInternetAddrEOS& Address);
 
+	/**
+	 * Returns the FSocketSubsytemEOS instance corresponding with the UWorld* parameter.
+	 * If no entry exists for the given UWorld*, it will add it if a valid instance is found.
+	 * 
+	 * @param InWorld UWorld* to be used in the search
+	 * @return The FSocketSubsystemEOS instance associated with InWorld
+	 */
+	FSocketSubsystemEOS* GetSocketSubsystemForWorld(UWorld* InWorld);
+
+private:
+	/** Removes the FSocketSubsystemEOS instance from SocketSubsystemEOSInstances and all related entries from SocketSubsystemEOSPerWorldMap */
+	void RemoveFromStaticContainers();
+
 private:
 #if WITH_EOS_SDK
 	EOS_HP2P P2PHandle;
@@ -96,4 +109,10 @@ private:
 
 	/** The last error we received */
 	ESocketErrors LastSocketError;
+
+	/** Static array of all instances of FSocketSubsystemEOS running concurrently (can be more than one in PIE) */
+	static TArray<FSocketSubsystemEOS*> SocketSubsystemEOSInstances;
+
+	/** Static map of world to corresponding FSocketSubsystemeOS instance */
+	static TMap<UWorld*, FSocketSubsystemEOS*> SocketSubsystemEOSPerWorldMap;
 };
diff --git a/Engine/Plugins/Online/SocketSubsystemEOS/Source/SocketSubsystemEOS/Public/SocketSubsystemEOSUtils.h b/Engine/Plugins/Online/SocketSubsystemEOS/Source/SocketSubsystemEOS/Public/SocketSubsystemEOSUtils.h
--- a/Engine/Plugins/Online/SocketSubsystemEOS/Source/SocketSubsystemEOS/Public/SocketSubsystemEOSUtils.h
+++ b/Engine/Plugins/Online/SocketSubsystemEOS/Source/SocketSubsystemEOS/Public/SocketSubsystemEOSUtils.h
@@ -21,5 +21,7 @@ public:
 	virtual EOS_ProductUserId GetLocalUserId() = 0;
 #endif
 	virtual FString GetSessionId() = 0;
+
+	virtual FName GetSubsystemInstanceName() = 0;
 };
 typedef TSharedPtr<ISocketSubsystemEOSUtils, ESPMode::ThreadSafe> ISocketSubsystemEOSUtilsPtr;
