Addition of instance name member to IOnlineServices, to be able to use GetWorldFromOnline in V2

diff --git a/Engine/Plugins/Online/Experimental/OnlineServices/Source/OnlineServicesCommon/Private/Online/OnlineServicesCommon.cpp b/Engine/Plugins/Online/Experimental/OnlineServices/Source/OnlineServicesCommon/Private/Online/OnlineServicesCommon.cpp
--- a/Engine/Plugins/Online/Experimental/OnlineServices/Source/OnlineServicesCommon/Private/Online/OnlineServicesCommon.cpp
+++ b/Engine/Plugins/Online/Experimental/OnlineServices/Source/OnlineServicesCommon/Private/Online/OnlineServicesCommon.cpp
@@ -11,9 +11,10 @@ namespace UE::Online {
 
 uint32 FOnlineServicesCommon::NextInstanceIndex = 0;
 
-FOnlineServicesCommon::FOnlineServicesCommon(const FString& InConfigName)
+FOnlineServicesCommon::FOnlineServicesCommon(const FString& InConfigName, FName InInstanceName)
 	: OpCache(InConfigName, *this)
 	, InstanceIndex(NextInstanceIndex++)
+	, InstanceName(InInstanceName)
 	, ConfigProvider(MakeUnique<FOnlineConfigProviderGConfig>(GEngineIni))
 	, ConfigName(InConfigName)
 	, SerialQueue(ParallelQueue)
@@ -79,6 +80,11 @@ TOnlineResult<FGetResolvedConnectString> FOnlineServicesCommon::GetResolvedConne
 	return TOnlineResult<FGetResolvedConnectString>(Errors::NotImplemented());
 }
 
+FName FOnlineServicesCommon::GetInstanceName() const
+{
+	return InstanceName;
+}
+
 void FOnlineServicesCommon::RegisterComponents()
 {
 }
diff --git a/Engine/Plugins/Online/Experimental/OnlineServices/Source/OnlineServicesCommon/Public/Online/OnlineServicesCommon.h b/Engine/Plugins/Online/Experimental/OnlineServices/Source/OnlineServicesCommon/Public/Online/OnlineServicesCommon.h
--- a/Engine/Plugins/Online/Experimental/OnlineServices/Source/OnlineServicesCommon/Public/Online/OnlineServicesCommon.h
+++ b/Engine/Plugins/Online/Experimental/OnlineServices/Source/OnlineServicesCommon/Public/Online/OnlineServicesCommon.h
@@ -27,7 +27,7 @@ class ONLINESERVICESCOMMON_API FOnlineServicesCommon
 public:
 	using Super = IOnlineServices;
 
-	FOnlineServicesCommon(const FString& InConfigName);
+	FOnlineServicesCommon(const FString& InConfigName, FName InInstanceName);
 	FOnlineServicesCommon(const FOnlineServicesCommon&) = delete;
 	FOnlineServicesCommon(FOnlineServicesCommon&&) = delete;
 	virtual ~FOnlineServicesCommon() {}
@@ -43,6 +43,7 @@ public:
 	virtual IConnectivityPtr GetConnectivityInterface() override;
 	virtual IPrivilegesPtr GetPrivilegesInterface() override;
 	virtual TOnlineResult<FGetResolvedConnectString> GetResolvedConnectString(FGetResolvedConnectString::Params&& Params) override;
+	virtual FName GetInstanceName() const override;
 
 	// FOnlineServicesCommon
 
@@ -310,6 +311,7 @@ protected:
 
 	static uint32 NextInstanceIndex;
 	uint32 InstanceIndex;
+	FName InstanceName;
 
 	FOnlineComponentRegistry Components;
 	TUniquePtr<IOnlineConfigProvider> ConfigProvider;
diff --git a/Engine/Plugins/Online/Experimental/OnlineServices/Source/OnlineServicesInterface/Private/Online/OnlineServicesRegistry.cpp b/Engine/Plugins/Online/Experimental/OnlineServices/Source/OnlineServicesInterface/Private/Online/OnlineServicesRegistry.cpp
--- a/Engine/Plugins/Online/Experimental/OnlineServices/Source/OnlineServicesInterface/Private/Online/OnlineServicesRegistry.cpp
+++ b/Engine/Plugins/Online/Experimental/OnlineServices/Source/OnlineServicesInterface/Private/Online/OnlineServicesRegistry.cpp
@@ -71,7 +71,7 @@ TSharedPtr<IOnlineServices> FOnlineServicesRegistry::GetNamedServicesInstance(EO
 		}
 		else
 		{
-			Services = CreateServices(OnlineServices);
+			Services = CreateServices(OnlineServices, InstanceName);
 			if (Services.IsValid())
 			{
 				NamedServiceInstances.FindOrAdd(OnlineServices).Add(InstanceName, Services.ToSharedRef());
@@ -93,14 +93,14 @@ void FOnlineServicesRegistry::DestroyNamedServicesInstance(EOnlineServices Onlin
 	}
 }
 
-TSharedPtr<IOnlineServices> FOnlineServicesRegistry::CreateServices(EOnlineServices OnlineServices)
+TSharedPtr<IOnlineServices> FOnlineServicesRegistry::CreateServices(EOnlineServices OnlineServices, FName InstanceName)
 {
 	TSharedPtr<IOnlineServices> Services;
 
 	FFactoryAndPriority* FactoryAndPriority = ServicesFactories.Find(OnlineServices);
 	if (FactoryAndPriority != nullptr)
 	{
-		Services = FactoryAndPriority->Factory->Create();
+		Services = FactoryAndPriority->Factory->Create(InstanceName);
 		Services->Init();
 	}
 
diff --git a/Engine/Plugins/Online/Experimental/OnlineServices/Source/OnlineServicesInterface/Public/Online/OnlineServices.h b/Engine/Plugins/Online/Experimental/OnlineServices/Source/OnlineServicesInterface/Public/Online/OnlineServices.h
--- a/Engine/Plugins/Online/Experimental/OnlineServices/Source/OnlineServicesInterface/Public/Online/OnlineServices.h
+++ b/Engine/Plugins/Online/Experimental/OnlineServices/Source/OnlineServicesInterface/Public/Online/OnlineServices.h
@@ -50,7 +50,7 @@ public:
 	 *
 	 */
 	virtual void Destroy() = 0;
-	
+
 	/**
 	 *
 	 */
@@ -97,6 +97,12 @@ public:
 	 * @return provider type
 	 */
 	virtual EOnlineServices GetServicesProvider() const = 0;
+
+	/**
+	 * Get the instance name of this instance. 
+	 * @return instance name
+	 */
+	virtual FName GetInstanceName() const = 0;
 };
 
 /**
diff --git a/Engine/Plugins/Online/Experimental/OnlineServices/Source/OnlineServicesInterface/Public/Online/OnlineServicesRegistry.h b/Engine/Plugins/Online/Experimental/OnlineServices/Source/OnlineServicesInterface/Public/Online/OnlineServicesRegistry.h
--- a/Engine/Plugins/Online/Experimental/OnlineServices/Source/OnlineServicesInterface/Public/Online/OnlineServicesRegistry.h
+++ b/Engine/Plugins/Online/Experimental/OnlineServices/Source/OnlineServicesInterface/Public/Online/OnlineServicesRegistry.h
@@ -27,7 +27,7 @@ public:
 	 *
 	 * @return Initialized IOnlineServices instance
 	 */
-	virtual TSharedPtr<IOnlineServices> Create() = 0;
+	virtual TSharedPtr<IOnlineServices> Create(FName InstanceName) = 0;
 };
 
 class FOnlineServicesRegistry
@@ -97,7 +97,7 @@ public:
 	 * 
 	 * @return The initialized IOnlineServices instance, or an invalid pointer if the OnlineServices is unavailable
 	 */
-	ONLINESERVICESINTERFACE_API TSharedPtr<IOnlineServices> CreateServices(EOnlineServices OnlineServices);
+	ONLINESERVICESINTERFACE_API TSharedPtr<IOnlineServices> CreateServices(EOnlineServices OnlineServices, FName InstanceName);
 
 	/**
 	 * Get list of all instantiated OnlineServices
diff --git a/Engine/Plugins/Online/Experimental/OnlineServicesEOS/Source/Private/Online/OnlineServicesEOS.cpp b/Engine/Plugins/Online/Experimental/OnlineServicesEOS/Source/Private/Online/OnlineServicesEOS.cpp
--- a/Engine/Plugins/Online/Experimental/OnlineServicesEOS/Source/Private/Online/OnlineServicesEOS.cpp
+++ b/Engine/Plugins/Online/Experimental/OnlineServicesEOS/Source/Private/Online/OnlineServicesEOS.cpp
@@ -45,8 +45,8 @@ END_ONLINE_STRUCT_META()
 
 /* Meta */ }
 
-FOnlineServicesEOS::FOnlineServicesEOS()
-	: FOnlineServicesCommon(TEXT("EOS"))
+FOnlineServicesEOS::FOnlineServicesEOS(FName InInstanceName)
+	: FOnlineServicesCommon(TEXT("EOS"), InInstanceName)
 {
 }
 
diff --git a/Engine/Plugins/Online/Experimental/OnlineServicesEOS/Source/Private/Online/SocketSubsystemEOSUtils_OnlineServicesEOS.cpp b/Engine/Plugins/Online/Experimental/OnlineServicesEOS/Source/Private/Online/SocketSubsystemEOSUtils_OnlineServicesEOS.cpp
--- a/Engine/Plugins/Online/Experimental/OnlineServicesEOS/Source/Private/Online/SocketSubsystemEOSUtils_OnlineServicesEOS.cpp
+++ b/Engine/Plugins/Online/Experimental/OnlineServicesEOS/Source/Private/Online/SocketSubsystemEOSUtils_OnlineServicesEOS.cpp
@@ -90,8 +90,7 @@ FString FSocketSubsystemEOSUtils_OnlineServicesEOS::GetSessionId()
 
 FName FSocketSubsystemEOSUtils_OnlineServicesEOS::GetSubsystemInstanceName()
 {
-	// TODO: OnlineServices still doesn't have functionality that matches GetInstanceName
-	return FName();
+	return ServicesEOS.GetInstanceName();
 }
 
 /* UE::Online */}
diff --git a/Engine/Plugins/Online/Experimental/OnlineServicesEOS/Source/Private/OnlineServicesEOSModule.cpp b/Engine/Plugins/Online/Experimental/OnlineServicesEOS/Source/Private/OnlineServicesEOSModule.cpp
--- a/Engine/Plugins/Online/Experimental/OnlineServicesEOS/Source/Private/OnlineServicesEOSModule.cpp
+++ b/Engine/Plugins/Online/Experimental/OnlineServicesEOS/Source/Private/OnlineServicesEOSModule.cpp
@@ -22,9 +22,9 @@ class FOnlineServicesFactoryEOS : public IOnlineServicesFactory
 {
 public:
 	virtual ~FOnlineServicesFactoryEOS() {}
-	virtual TSharedPtr<IOnlineServices> Create() override
+	virtual TSharedPtr<IOnlineServices> Create(FName InInstanceName) override
 	{
-		return MakeShared<FOnlineServicesEOS>();
+		return MakeShared<FOnlineServicesEOS>(InInstanceName);
 	}
 protected:
 };
diff --git a/Engine/Plugins/Online/Experimental/OnlineServicesEOS/Source/Public/Online/OnlineServicesEOS.h b/Engine/Plugins/Online/Experimental/OnlineServicesEOS/Source/Public/Online/OnlineServicesEOS.h
--- a/Engine/Plugins/Online/Experimental/OnlineServicesEOS/Source/Public/Online/OnlineServicesEOS.h
+++ b/Engine/Plugins/Online/Experimental/OnlineServicesEOS/Source/Public/Online/OnlineServicesEOS.h
@@ -22,7 +22,7 @@ class ONLINESERVICESEOS_API FOnlineServicesEOS : public FOnlineServicesCommon
 public:
 	using Super = FOnlineServicesCommon;
 
-	FOnlineServicesEOS();
+	FOnlineServicesEOS(FName InInstanceName);
 	virtual void RegisterComponents() override;
 	virtual void Initialize() override;
 	virtual TOnlineResult<FGetResolvedConnectString> GetResolvedConnectString(FGetResolvedConnectString::Params&& Params) override;
diff --git a/Engine/Plugins/Online/Experimental/OnlineServicesNull/Source/Private/Online/OnlineServicesNull.cpp b/Engine/Plugins/Online/Experimental/OnlineServicesNull/Source/Private/Online/OnlineServicesNull.cpp
--- a/Engine/Plugins/Online/Experimental/OnlineServicesNull/Source/Private/Online/OnlineServicesNull.cpp
+++ b/Engine/Plugins/Online/Experimental/OnlineServicesNull/Source/Private/Online/OnlineServicesNull.cpp
@@ -23,8 +23,8 @@ END_ONLINE_STRUCT_META()
 
  /*Meta*/ }
 
-FOnlineServicesNull::FOnlineServicesNull()
-	: FOnlineServicesCommon(TEXT("Null"))
+FOnlineServicesNull::FOnlineServicesNull(FName InInstanceName)
+	: FOnlineServicesCommon(TEXT("Null"), InInstanceName)
 {
 }
 
diff --git a/Engine/Plugins/Online/Experimental/OnlineServicesNull/Source/Private/Online/OnlineServicesNullModule.cpp b/Engine/Plugins/Online/Experimental/OnlineServicesNull/Source/Private/Online/OnlineServicesNullModule.cpp
--- a/Engine/Plugins/Online/Experimental/OnlineServicesNull/Source/Private/Online/OnlineServicesNullModule.cpp
+++ b/Engine/Plugins/Online/Experimental/OnlineServicesNull/Source/Private/Online/OnlineServicesNullModule.cpp
@@ -22,9 +22,9 @@ class FOnlineServicesFactoryNull : public IOnlineServicesFactory
 {
 public:
 	virtual ~FOnlineServicesFactoryNull() {}
-	virtual TSharedPtr<IOnlineServices> Create() override
+	virtual TSharedPtr<IOnlineServices> Create(FName InInstanceName) override
 	{
-		return MakeShared<FOnlineServicesNull>();
+		return MakeShared<FOnlineServicesNull>(InInstanceName);
 	}
 protected:
 };
diff --git a/Engine/Plugins/Online/Experimental/OnlineServicesNull/Source/Public/Online/OnlineServicesNull.h b/Engine/Plugins/Online/Experimental/OnlineServicesNull/Source/Public/Online/OnlineServicesNull.h
--- a/Engine/Plugins/Online/Experimental/OnlineServicesNull/Source/Public/Online/OnlineServicesNull.h
+++ b/Engine/Plugins/Online/Experimental/OnlineServicesNull/Source/Public/Online/OnlineServicesNull.h
@@ -11,7 +11,7 @@ class ONLINESERVICESNULL_API FOnlineServicesNull : public FOnlineServicesCommon
 public:
 	using Super = FOnlineServicesCommon;
 
-	FOnlineServicesNull();
+	FOnlineServicesNull(FName InInstanceName);
 	virtual void RegisterComponents() override;
 	virtual void Initialize() override;
 	virtual TOnlineResult<FGetResolvedConnectString> GetResolvedConnectString(FGetResolvedConnectString::Params&& Params) override;
diff --git a/Engine/Plugins/Online/Experimental/OnlineServicesOSSAdapter/Source/Private/Online/OnlineServicesOSSAdapter.cpp b/Engine/Plugins/Online/Experimental/OnlineServicesOSSAdapter/Source/Private/Online/OnlineServicesOSSAdapter.cpp
--- a/Engine/Plugins/Online/Experimental/OnlineServicesOSSAdapter/Source/Private/Online/OnlineServicesOSSAdapter.cpp
+++ b/Engine/Plugins/Online/Experimental/OnlineServicesOSSAdapter/Source/Private/Online/OnlineServicesOSSAdapter.cpp
@@ -12,8 +12,8 @@
 
 namespace UE::Online {
 
-FOnlineServicesOSSAdapter::FOnlineServicesOSSAdapter(EOnlineServices InServicesType, const FString& InConfigName, IOnlineSubsystem* InSubsystem)
-	: FOnlineServicesCommon(InConfigName)
+FOnlineServicesOSSAdapter::FOnlineServicesOSSAdapter(EOnlineServices InServicesType, const FString& InConfigName, FName InInstanceName, IOnlineSubsystem* InSubsystem)
+	: FOnlineServicesCommon(InConfigName, InInstanceName)
 	, ServicesType(InServicesType)
 	, Subsystem(InSubsystem)
 {
diff --git a/Engine/Plugins/Online/Experimental/OnlineServicesOSSAdapter/Source/Private/Online/OnlineServicesOSSAdapterModule.cpp b/Engine/Plugins/Online/Experimental/OnlineServicesOSSAdapter/Source/Private/Online/OnlineServicesOSSAdapterModule.cpp
--- a/Engine/Plugins/Online/Experimental/OnlineServicesOSSAdapter/Source/Private/Online/OnlineServicesOSSAdapterModule.cpp
+++ b/Engine/Plugins/Online/Experimental/OnlineServicesOSSAdapter/Source/Private/Online/OnlineServicesOSSAdapterModule.cpp
@@ -59,12 +59,12 @@ public:
 	}
 
 	virtual ~FOnlineServicesFactoryOSSAdapter() {}
-	virtual TSharedPtr<IOnlineServices> Create() override
+	virtual TSharedPtr<IOnlineServices> Create(FName InInstanceName) override
 	{
 		IOnlineSubsystem* Subsystem = IOnlineSubsystem::Get(Config.OnlineSubsystem);
 		if (Subsystem != nullptr)
 		{
-			return MakeShared<FOnlineServicesOSSAdapter>(Config.Service, Config.ConfigName, Subsystem);
+			return MakeShared<FOnlineServicesOSSAdapter>(Config.Service, Config.ConfigName, InInstanceName, Subsystem);
 		}
 		else
 		{
diff --git a/Engine/Plugins/Online/Experimental/OnlineServicesOSSAdapter/Source/Public/Online/OnlineServicesOSSAdapter.h b/Engine/Plugins/Online/Experimental/OnlineServicesOSSAdapter/Source/Public/Online/OnlineServicesOSSAdapter.h
--- a/Engine/Plugins/Online/Experimental/OnlineServicesOSSAdapter/Source/Public/Online/OnlineServicesOSSAdapter.h
+++ b/Engine/Plugins/Online/Experimental/OnlineServicesOSSAdapter/Source/Public/Online/OnlineServicesOSSAdapter.h
@@ -15,7 +15,7 @@ class FOnlineServicesOSSAdapter : public FOnlineServicesCommon
 public:
 	using Super = FOnlineServicesCommon;
 
-	ONLINESERVICESOSSADAPTER_API FOnlineServicesOSSAdapter(EOnlineServices InServicesType, const FString& InConfigName, IOnlineSubsystem* InSubsystem);
+	ONLINESERVICESOSSADAPTER_API FOnlineServicesOSSAdapter(EOnlineServices InServicesType, const FString& InConfigName, FName InInstanceName, IOnlineSubsystem* InSubsystem);
 
 	ONLINESERVICESOSSADAPTER_API virtual void RegisterComponents() override;
 	ONLINESERVICESOSSADAPTER_API virtual void Initialize() override;
diff --git a/Engine/Plugins/Online/SocketSubsystemEOS/Source/SocketSubsystemEOS/Private/NetDriverEOSBase.cpp b/Engine/Plugins/Online/SocketSubsystemEOS/Source/SocketSubsystemEOS/Private/NetDriverEOSBase.cpp
--- a/Engine/Plugins/Online/SocketSubsystemEOS/Source/SocketSubsystemEOS/Private/NetDriverEOSBase.cpp
+++ b/Engine/Plugins/Online/SocketSubsystemEOS/Source/SocketSubsystemEOS/Private/NetDriverEOSBase.cpp
@@ -7,7 +7,6 @@
 #include "Misc/EngineVersionComparison.h"
 #include "EOSSharedTypes.h"
 #include "Engine/Engine.h"
-#include "Misc/ConfigCacheIni.h"
 
 bool UNetDriverEOSBase::IsAvailable() const
 {
@@ -178,22 +177,9 @@ ISocketSubsystem* UNetDriverEOSBase::GetSocketSubsystem()
 	}
 	else
 	{
+		UWorld* CurrentWorld = FindWorld();
 		FSocketSubsystemEOS* DefaultSocketSubsystem = static_cast<FSocketSubsystemEOS*>(ISocketSubsystem::Get(EOS_SOCKETSUBSYSTEM));
-
-		bool bUseOnlineServicesV2 = false;
-		GConfig->GetBool(TEXT("/Script/Engine.OnlineEngineInterface"), TEXT("bUseOnlineServicesV2"), bUseOnlineServicesV2, GEngineIni);
-
-		if (bUseOnlineServicesV2)
-		{
-			// V2 still doesn't have the logic to process several instances, so we'll return the default one
-			return DefaultSocketSubsystem;
-		}
-		else
-		{
-			// We may have several instances of FSocketSubsystemEOS if running in PIE mode, so we need to pick the right one for the current world
-			UWorld* CurrentWorld = FindWorld();
-			return DefaultSocketSubsystem->GetSocketSubsystemForWorld(CurrentWorld);
-		}
+		return DefaultSocketSubsystem->GetSocketSubsystemForWorld(CurrentWorld);
 	}
 }
 
