Added logic to load SocketSubsystemEOS Relay mode via config

diff --git a/Engine/Plugins/Online/SocketSubsystemEOS/Source/SocketSubsystemEOS/Private/SocketSubsystemEOS.cpp b/Engine/Plugins/Online/SocketSubsystemEOS/Source/SocketSubsystemEOS/Private/SocketSubsystemEOS.cpp
--- a/Engine/Plugins/Online/SocketSubsystemEOS/Source/SocketSubsystemEOS/Private/SocketSubsystemEOS.cpp
+++ b/Engine/Plugins/Online/SocketSubsystemEOS/Source/SocketSubsystemEOS/Private/SocketSubsystemEOS.cpp
@@ -13,8 +13,33 @@
 
 #if WITH_EOS_SDK
 	#include "eos_sdk.h"
+	#include "eos_p2p.h"
 #endif
 
+namespace
+{
+
+	static EOS_ERelayControl LexFromString(const FString& Value)
+	{
+		if (Value == TEXT("NoRelays"))
+		{
+			return EOS_ERelayControl::EOS_RC_NoRelays;
+		}
+		if (Value == TEXT("AllowRelays"))
+		{
+			return EOS_ERelayControl::EOS_RC_AllowRelays;
+		}
+		if (Value == TEXT("ForceRelays"))
+		{
+			return EOS_ERelayControl::EOS_RC_ForceRelays;
+		}
+
+		checkNoEntry();
+		return EOS_ERelayControl::EOS_RC_AllowRelays;
+	}
+
+}
+
 TArray<FSocketSubsystemEOS*> FSocketSubsystemEOS::SocketSubsystemEOSInstances;
 TMap<UWorld*, FSocketSubsystemEOS*> FSocketSubsystemEOS::SocketSubsystemEOSPerWorldMap;
 
@@ -68,6 +93,25 @@ bool FSocketSubsystemEOS::Init(FString& Error)
 	FSocketSubsystemModule& SocketSubsystem = FModuleManager::LoadModuleChecked<FSocketSubsystemModule>("Sockets");
 	SocketSubsystem.RegisterSocketSubsystem(EOS_SOCKETSUBSYSTEM, this, false);
 
+#if WITH_EOS_SDK
+	FString RelayControlStr;
+	GConfig->GetString(TEXT("SocketSubsystemEOS"), TEXT("RelayControl"), RelayControlStr, GEngineIni);
+	if (!RelayControlStr.IsEmpty())
+	{
+		EOS_P2P_SetRelayControlOptions Options = {};
+		Options.ApiVersion = EOS_P2P_SETRELAYCONTROL_API_LATEST;
+		static_assert(EOS_P2P_SETRELAYCONTROL_API_LATEST == 1, "EOS_P2P_SetRelayControlOptions updated, check new fields");
+
+		Options.RelayControl = LexFromString(RelayControlStr);
+
+		EOS_EResult Result = EOS_P2P_SetRelayControl(P2PHandle, &Options);
+		if (Result != EOS_EResult::EOS_Success && Result != EOS_EResult::EOS_NoChange)
+		{
+			UE_LOG(LogSocketSubsystemEOS, Warning, TEXT("[FSocketSubsystemEOS::Init] EOS_P2P_SetRelayControl failed with result [%s]"), UTF8_TO_TCHAR(EOS_EResult_ToString(Result)));
+		}
+	}
+#endif
+
 	return true;
 }
 
