- Remove NpLog support

diff --git a/Engine/Plugins/Online/SocketSubsystemEOS/Source/SocketSubsystemEOS/Private/SocketEOS.cpp b/Engine/Plugins/Online/SocketSubsystemEOS/Source/SocketSubsystemEOS/Private/SocketEOS.cpp
--- a/Engine/Plugins/Online/SocketSubsystemEOS/Source/SocketSubsystemEOS/Private/SocketEOS.cpp
+++ b/Engine/Plugins/Online/SocketSubsystemEOS/Source/SocketSubsystemEOS/Private/SocketEOS.cpp
@@ -8,52 +8,6 @@
 	#include "eos_p2p.h"
 #endif
 
-#if WANTS_NP_LOGGING
-
-#include "Windows/AllowWindowsPlatformTypes.h"
-THIRD_PARTY_INCLUDES_START
-#include <Windows.h>
-THIRD_PARTY_INCLUDES_END
-
-const TCHAR* GetLogPrefix()
-{
-	static FString Prefix;
-	if (Prefix.Len() == 0)
-	{
-		FParse::Value(FCommandLine::Get(), TEXT("LogPrefix="), Prefix);
-		if (Prefix.Len() == 0)
-		{
-			Prefix = TEXT("Unknown");
-		}
-	}
-	return *Prefix;
-}
-
-void NpLog(const TCHAR* Msg)
-{
-	static HWND EditWindow = NULL;
-	// Get the edit window so we can send messages to it
-	if (EditWindow == NULL)
-	{
-		HWND MainWindow = FindWindowW(NULL, L"Untitled - Notepad");
-		if (MainWindow == NULL)
-		{
-			MainWindow = FindWindowW(NULL, L"*Untitled - Notepad");
-		}
-		if (MainWindow != NULL)
-		{
-			EditWindow = FindWindowExW(MainWindow, NULL, L"Edit", NULL);
-		}
-	}
-	if (EditWindow != NULL)
-	{
-		SendMessageW(EditWindow, EM_REPLACESEL, TRUE, (LPARAM)Msg);
-	}
-}
-
-#include "Windows/HideWindowsPlatformTypes.h"
-#endif
-
 FSocketEOS::FSocketEOS(FSocketSubsystemEOS& InSocketSubsystem, const FString& InSocketDescription)
 	: FSocket(ESocketType::SOCKTYPE_Datagram, InSocketDescription, NAME_None)
 	, SocketSubsystem(InSocketSubsystem)
@@ -120,7 +74,6 @@ bool FSocketEOS::Close()
 		EOS_EResult Result = EOS_P2P_CloseConnections(SocketSubsystem.GetP2PHandle(), &Options);
 
 		UE_LOG(LogSocketSubsystemEOS, Log, TEXT("Closing socket (%s) with result (%s)"), *LocalAddress.ToString(true), *LexToString(Result));
-		NP_LOG(TEXT("[%s] - Closing socket (%s) with result (%s)\r\n"), GetLogPrefix(), *LocalAddress.ToString(true), *LexToString(Result));
 
 		ClosedRemotes.Empty();
 	}
@@ -172,7 +125,6 @@ bool FSocketEOS::Bind(const FInternetAddr& Addr)
 	LocalAddress.SetLocalUserId(LocalUserId);
 
 	UE_LOG(LogSocketSubsystemEOS, Verbose, TEXT("Successfully bound socket to address (%s)"), *LocalAddress.ToString(true));
-	NP_LOG(TEXT("[%s] - Successfully bound socket to address (%s)\r\n"), GetLogPrefix(), *LocalAddress.ToString(true));
 	return true;
 }
 
@@ -240,12 +192,10 @@ bool FSocketEOS::Listen(int32)
 			if (AcceptResult == EOS_EResult::EOS_Success)
 			{
 				UE_LOG(LogSocketSubsystemEOS, Verbose, TEXT("Accepting connection request from (%s) on socket (%s)"), *RemoteUser, UTF8_TO_TCHAR(Info->SocketId->SocketName));
-				NP_LOG(TEXT("[%s] - Accepting connection request from (%s) on socket (%s)\r\n"), GetLogPrefix(), *RemoteUser, UTF8_TO_TCHAR(Info->SocketId->SocketName));
 			}
 			else
 			{
 				UE_LOG(LogSocketSubsystemEOS, Error, TEXT("EOS_P2P_AcceptConnection from (%s) on socket (%s) failed with (%s)"), *RemoteUser, UTF8_TO_TCHAR(Info->SocketId->SocketName), *LexToString(AcceptResult));
-				NP_LOG(TEXT("[%s] - EOS_P2P_AcceptConnection from (%s) on socket (%s) failed with (%s)\r\n"), GetLogPrefix(), *RemoteUser, UTF8_TO_TCHAR(Info->SocketId->SocketName), *LexToString(AcceptResult));
 			}
 		}
 		else
@@ -406,7 +356,6 @@ bool FSocketEOS::SendTo(const uint8* Data, int32 Count, int32& OutBytesSent, con
 	Options.DataLengthBytes = Count;
 	Options.Data = Data;
 	EOS_EResult Result = EOS_P2P_SendPacket(SocketSubsystem.GetP2PHandle(), &Options);
-	NP_LOG(TEXT("[%s] - EOS_P2P_SendPacket() to (%s) result code = (%s)\r\n"), GetLogPrefix(), *Destination.ToString(true), *LexToString(Result));
 	if (Result != EOS_EResult::EOS_Success)
 	{
 		UE_LOG(LogSocketSubsystemEOS, Error, TEXT("Unable to send data to (%s) result code = (%s)"), *Destination.ToString(true), *LexToString(Result));
@@ -465,7 +414,6 @@ bool FSocketEOS::RecvFrom(uint8* Data, int32 BufferSize, int32& BytesRead, FInte
 	EOS_P2P_SocketId SocketId;
 	
 	EOS_EResult Result = EOS_P2P_ReceivePacket(SocketSubsystem.GetP2PHandle(), &Options, &RemoteUserId, &SocketId, &Channel, Data, (uint32*)&BytesRead);
-	NP_LOG(TEXT("[%s] - EOS_P2P_ReceivePacket() for user (%s) and channel (%d) with result code = (%s)\r\n"), GetLogPrefix(), *MakeStringFromProductUserId(LocalAddress.GetLocalUserId()), Channel, *LexToString(Result));
 	if (Result == EOS_EResult::EOS_NotFound)
 	{
 		// No data to read
@@ -487,7 +435,6 @@ bool FSocketEOS::RecvFrom(uint8* Data, int32 BufferSize, int32& BytesRead, FInte
 	SourceAddress.SetSocketName(SocketId.SocketName);
 	SourceAddress.SetChannel(Channel);
 
-	NP_LOG(TEXT("[%s] - EOS_P2P_ReceivePacket() of size (%d) from (%s)\r\n"), GetLogPrefix(), BytesRead, *SourceAddress.ToString(true));
 	return true;
 #else
 	return false;
@@ -643,7 +590,6 @@ bool FSocketEOS::Close(const FInternetAddrEOS& RemoteAddress)
 	Options.SocketId = &SocketId;
 
 	EOS_EResult Result = EOS_P2P_CloseConnection(SocketSubsystem.GetP2PHandle(), &Options);
-	NP_LOG(TEXT("[%s] - EOS_P2P_CloseConnection() with remote address RemoteAddress (%s) result code (%s)\r\n"), GetLogPrefix(), *RemoteAddress.ToString(true), *LexToString(Result));
 	if (Result != EOS_EResult::EOS_Success)
 	{
 		UE_LOG(LogSocketSubsystemEOS, Error, TEXT("Unable to close socket with remote address RemoteAddress (%s) due to error (%s)"), *RemoteAddress.ToString(true), *LexToString(Result));
@@ -685,7 +631,6 @@ void FSocketEOS::RegisterClosedNotification()
 		FInternetAddrEOS RemoteAddress(Info->RemoteUserId, Info->SocketId->SocketName, LocalAddress.GetChannel());
 		RemoteAddress.SetLocalUserId(LocalAddress.GetLocalUserId());
 		ClosedRemotes.Add(RemoteAddress);
-		NP_LOG(TEXT("[%s] - Close connection received for remote address (%s)\r\n"), GetLogPrefix(), *RemoteAddress.ToString(true));
 	};
 	ClosedNotifyId = EOS_P2P_AddNotifyPeerConnectionClosed(SocketSubsystem.GetP2PHandle(), &Options, ClosedNotifyCallback, ClosedNotifyCallback->GetCallbackPtr());
 #endif
diff --git a/Engine/Plugins/Online/SocketSubsystemEOS/Source/SocketSubsystemEOS/Public/SocketEOS.h b/Engine/Plugins/Online/SocketSubsystemEOS/Source/SocketSubsystemEOS/Public/SocketEOS.h
--- a/Engine/Plugins/Online/SocketSubsystemEOS/Source/SocketSubsystemEOS/Public/SocketEOS.h
+++ b/Engine/Plugins/Online/SocketSubsystemEOS/Source/SocketSubsystemEOS/Public/SocketEOS.h
@@ -10,16 +10,6 @@
 class FOnlineSubsystemEOS;
 class FSocketSubsystemEOS;
 
-#define WANTS_NP_LOGGING 0
-
-#if WANTS_NP_LOGGING
-	#define NP_LOG(Msg, ...) NpLog(*FString::Printf(Msg, __VA_ARGS__))
-
-	void NpLog(const TCHAR* Msg);
-#else
-	#define NP_LOG(Msg, ...)
-#endif
-
 #if WITH_EOS_SDK
 	#include "eos_p2p_types.h"
 #endif
