Added configurable choice of Packet Reliability Type in EOS Sockets

diff --git a/Engine/Plugins/Online/EOSShared/Source/EOSShared/Private/EOSShared.cpp b/Engine/Plugins/Online/EOSShared/Source/EOSShared/Private/EOSShared.cpp
--- a/Engine/Plugins/Online/EOSShared/Source/EOSShared/Private/EOSShared.cpp
+++ b/Engine/Plugins/Online/EOSShared/Source/EOSShared/Private/EOSShared.cpp
@@ -12,6 +12,7 @@
 #include "eos_presence_types.h"
 #include "eos_rtc_types.h"
 #include "eos_userinfo_types.h"
+#include "eos_p2p_types.h"
 
 DEFINE_LOG_CATEGORY(LogEOSSDK);
 
@@ -471,6 +472,29 @@ bool LexFromString(EOS_EIntegratedPlatformManagementFlags& OutEnum, const TCHAR*
 	return true;
 }
 
+bool LexFromString(EOS_EPacketReliability& OutEnum, const TCHAR* InString)
+{
+	if (FCString::Stricmp(InString, TEXT("UnreliableUnordered")) == 0)
+	{
+		OutEnum = EOS_EPacketReliability::EOS_PR_UnreliableUnordered;
+	}
+	else if (FCString::Stricmp(InString, TEXT("ReliableUnordered")) == 0)
+	{
+		OutEnum = EOS_EPacketReliability::EOS_PR_ReliableUnordered;
+	}
+	else if (FCString::Stricmp(InString, TEXT("ReliableOrdered")) == 0)
+	{
+		OutEnum = EOS_EPacketReliability::EOS_PR_ReliableOrdered;
+	}
+	else
+	{
+		checkNoEntry();
+		return false;
+	}
+
+	return true;
+}
+
 FString LexToString(const EOS_RTC_Option& Option)
 {
 	UE_EOS_CHECK_API_MISMATCH(EOS_RTC_OPTION_API_LATEST, 1);
diff --git a/Engine/Plugins/Online/EOSShared/Source/EOSShared/Public/EOSShared.h b/Engine/Plugins/Online/EOSShared/Source/EOSShared/Public/EOSShared.h
--- a/Engine/Plugins/Online/EOSShared/Source/EOSShared/Public/EOSShared.h
+++ b/Engine/Plugins/Online/EOSShared/Source/EOSShared/Public/EOSShared.h
@@ -54,6 +54,7 @@ EOS_ENUM_FORWARD_DECL(EOS_ELoginCredentialType);
 EOS_ENUM_FORWARD_DECL(EOS_ENetworkStatus);
 EOS_ENUM_FORWARD_DECL(EOS_Presence_EStatus);
 EOS_ENUM_FORWARD_DECL(EOS_UI_EInputStateButtonFlags);
+EOS_ENUM_FORWARD_DECL(EOS_EPacketReliability);
 #undef EOS_ENUM_FORWARD_DECL
 
 #define EOS_STRUCT_FORWARD_DECL(name) extern "C" typedef struct _tag ## name name;
@@ -93,6 +94,7 @@ EOSSHARED_API bool LexFromString(EOS_ELoginCredentialType& OutEnum, const TCHAR*
 EOSSHARED_API bool LexFromString(EOS_ERTCBackgroundMode& OutEnum, const TCHAR* InString);
 EOSSHARED_API bool LexFromString(EOS_UI_EInputStateButtonFlags& OutEnum, const TCHAR* InString);
 EOSSHARED_API bool LexFromString(EOS_EIntegratedPlatformManagementFlags& OutEnum, const TCHAR* InString);
+EOSSHARED_API bool LexFromString(EOS_EPacketReliability& OutEnum, const TCHAR* InString);
 
 /** Extracts the display name FString from a EOS_UserInfo_BestDisplayName using the following logic: Nickname > DisplayNameSanitized > DisplayName */
 EOSSHARED_API FString GetBestDisplayNameStr(const EOS_UserInfo_BestDisplayName& BestDisplayName);
diff --git a/Engine/Plugins/Online/SocketSubsystemEOS/Source/SocketSubsystemEOS/Private/SocketEOS.cpp b/Engine/Plugins/Online/SocketSubsystemEOS/Source/SocketSubsystemEOS/Private/SocketEOS.cpp
--- a/Engine/Plugins/Online/SocketSubsystemEOS/Source/SocketSubsystemEOS/Private/SocketEOS.cpp
+++ b/Engine/Plugins/Online/SocketSubsystemEOS/Source/SocketSubsystemEOS/Private/SocketEOS.cpp
@@ -3,6 +3,7 @@
 #include "SocketEOS.h"
 #include "SocketSubsystemEOS.h"
 #include "EOSShared.h"
+#include "Misc/ConfigCacheIni.h"
 
 #if WITH_EOS_SDK
 	#include "eos_p2p.h"
@@ -20,6 +21,18 @@ FSocketEOS::FSocketEOS(FSocketSubsystemEOS& InSocketSubsystem, const FString& In
 #endif
 {
 	CallbackAliveTracker = MakeShared<FCallbackBase>();
+
+#if WITH_EOS_SDK
+	FString PacketReliabilityTypeStr;
+	if (GConfig->GetString(TEXT("SocketSubsystemEOS"), TEXT("DefaultPacketReliabilityType"), PacketReliabilityTypeStr, GEngineIni))
+	{
+		EOS_EPacketReliability PacketReliabilityType;
+		if (LexFromString(PacketReliabilityType, *PacketReliabilityTypeStr))
+		{
+			PacketReliability = PacketReliabilityType;
+		}
+	}
+#endif
 }
 
 FSocketEOS::~FSocketEOS()
@@ -351,7 +364,7 @@ bool FSocketEOS::SendTo(const uint8* Data, int32 Count, int32& OutBytesSent, con
 	Options.RemoteUserId = DestinationAddress.GetRemoteUserId();
 	Options.SocketId = &SocketId;
 	Options.bAllowDelayedDelivery = EOS_TRUE;
-	Options.Reliability = EOS_EPacketReliability::EOS_PR_UnreliableUnordered;
+	Options.Reliability = PacketReliability;
 	Options.Channel = DestinationAddress.GetChannel();
 	Options.DataLengthBytes = Count;
 	Options.Data = Data;
diff --git a/Engine/Plugins/Online/SocketSubsystemEOS/Source/SocketSubsystemEOS/Public/SocketEOS.h b/Engine/Plugins/Online/SocketSubsystemEOS/Source/SocketSubsystemEOS/Public/SocketEOS.h
--- a/Engine/Plugins/Online/SocketSubsystemEOS/Source/SocketSubsystemEOS/Public/SocketEOS.h
+++ b/Engine/Plugins/Online/SocketSubsystemEOS/Source/SocketSubsystemEOS/Public/SocketEOS.h
@@ -93,5 +93,8 @@ private:
 
 	FClosedNotifyCallback* ClosedNotifyCallback;
 	EOS_NotificationId ClosedNotifyId;
+
+	/** Default packet reliability type used in EOS_P2P_SendPacketOptions */
+	EOS_EPacketReliability PacketReliability = EOS_EPacketReliability::EOS_PR_UnreliableUnordered;
 #endif
 };
diff --git a/Engine/Plugins/Online/SocketSubsystemEOS/Source/SocketSubsystemEOS/SocketSubsystemEOS.Build.cs b/Engine/Plugins/Online/SocketSubsystemEOS/Source/SocketSubsystemEOS/SocketSubsystemEOS.Build.cs
--- a/Engine/Plugins/Online/SocketSubsystemEOS/Source/SocketSubsystemEOS/SocketSubsystemEOS.Build.cs
+++ b/Engine/Plugins/Online/SocketSubsystemEOS/Source/SocketSubsystemEOS/SocketSubsystemEOS.Build.cs
@@ -27,6 +27,7 @@ public class SocketSubsystemEOS : ModuleRules
 			{
 				"CoreOnline",
 				"CoreUObject",
+				"EOSShared",
 				"EOSSDK"
 			}
 		);
