EOS Packet Reliability config name change (missing from previous CL)

diff --git a/Engine/Plugins/Online/SocketSubsystemEOS/Source/SocketSubsystemEOS/Private/SocketEOS.cpp b/Engine/Plugins/Online/SocketSubsystemEOS/Source/SocketSubsystemEOS/Private/SocketEOS.cpp
--- a/Engine/Plugins/Online/SocketSubsystemEOS/Source/SocketSubsystemEOS/Private/SocketEOS.cpp
+++ b/Engine/Plugins/Online/SocketSubsystemEOS/Source/SocketSubsystemEOS/Private/SocketEOS.cpp
@@ -24,7 +24,7 @@ FSocketEOS::FSocketEOS(FSocketSubsystemEOS& InSocketSubsystem, const FString& In
 
 #if WITH_EOS_SDK
 	FString PacketReliabilityTypeStr;
-	if (GConfig->GetString(TEXT("SocketSubsystemEOS"), TEXT("DefaultPacketReliabilityType"), PacketReliabilityTypeStr, GEngineIni))
+	if (GConfig->GetString(TEXT("SocketSubsystemEOS"), TEXT("PacketReliabilityType"), PacketReliabilityTypeStr, GEngineIni))
 	{
 		EOS_EPacketReliability PacketReliabilityType;
 		if (LexFromString(PacketReliabilityType, *PacketReliabilityTypeStr))
