Addition of missing change after base class removal.

diff --git a/Engine/Plugins/Online/SocketSubsystemEOS/Source/SocketSubsystemEOS/Private/NetDriverEOS.cpp b/Engine/Plugins/Online/SocketSubsystemEOS/Source/SocketSubsystemEOS/Private/NetDriverEOS.cpp
--- a/Engine/Plugins/Online/SocketSubsystemEOS/Source/SocketSubsystemEOS/Private/NetDriverEOS.cpp
+++ b/Engine/Plugins/Online/SocketSubsystemEOS/Source/SocketSubsystemEOS/Private/NetDriverEOS.cpp
@@ -90,7 +90,7 @@ bool UNetDriverEOS::InitBase(bool bInitAsClient, FNetworkNotify* InNotify, const
 	TSharedRef<FInternetAddrEOS> EOSLocalAddress = StaticCastSharedRef<FInternetAddrEOS>(LocalAddress);
 	EOSLocalAddress->SetChannel((uint8)URL.Port);
 	// Set our net driver name so we don't accept connections across net driver types
-	EOSLocalAddress->SetSocketName(NetDriverName.ToString());
+	EOSLocalAddress->SetSocketName(GetNetDriverDefinition().ToString());
 
 	static_cast<FSocketEOS*>(GetSocket())->SetLocalAddress(*EOSLocalAddress);
 
