Rebind connection request handler when network goes from offline to online, and cleanup objects when network goes from online to offline.

diff --git a/Engine/Plugins/Online/EOSShared/Source/EOSShared/Private/EOSSDKManager.cpp b/Engine/Plugins/Online/EOSShared/Source/EOSShared/Private/EOSSDKManager.cpp
--- a/Engine/Plugins/Online/EOSShared/Source/EOSShared/Private/EOSSDKManager.cpp
+++ b/Engine/Plugins/Online/EOSShared/Source/EOSShared/Private/EOSSDKManager.cpp
@@ -962,6 +962,7 @@ void FEOSSDKManager::OnNetworkConnectionStatusChanged(ENetworkConnectionStatus L
 	for (const TPair<EOS_HPlatform, IEOSPlatformHandleWeakPtr>& Entry : ActivePlatforms)
 	{
 		EOS_Platform_SetNetworkStatus(Entry.Key, NewNetworkStatus);
+		OnNetworkStatusChanged.Broadcast(OldNetworkStatus, NewNetworkStatus);
 	}
 }
 
@@ -1170,6 +1171,30 @@ bool FEOSSDKManager::Exec_Runtime(class UWorld* InWorld, const TCHAR* Cmd, FOutp
 	{
 		LogInfo();
 	}
+	else if (FParse::Command(&Cmd, TEXT("DISABLENETWORK")))
+	{
+		for (const TPair<EOS_HPlatform, IEOSPlatformHandleWeakPtr>& Entry : ActivePlatforms)
+		{
+			EOS_ENetworkStatus CurrentStatus = EOS_Platform_GetNetworkStatus(Entry.Key);
+
+			if (CurrentStatus == EOS_ENetworkStatus::EOS_NS_Online)
+			{
+				OnNetworkConnectionStatusChanged(ENetworkConnectionStatus::Connected, ENetworkConnectionStatus::Local); 
+			}
+		}
+	}
+	else if (FParse::Command(&Cmd, TEXT("ENABLENETWORK")))
+	{
+		for (const TPair<EOS_HPlatform, IEOSPlatformHandleWeakPtr>& Entry : ActivePlatforms)
+		{
+			EOS_ENetworkStatus CurrentStatus = EOS_Platform_GetNetworkStatus(Entry.Key);
+
+			if (CurrentStatus == EOS_ENetworkStatus::EOS_NS_Offline)
+			{
+				OnNetworkConnectionStatusChanged(ENetworkConnectionStatus::Local, ENetworkConnectionStatus::Connected);
+			}
+		}
+	}
 	else
 	{
 		UE_LOG(LogEOSShared, Warning, TEXT("%hs Unknown exec command: %s]"), __FUNCTION__, Cmd);
diff --git a/Engine/Plugins/Online/EOSShared/Source/EOSShared/Public/IEOSSDKManager.h b/Engine/Plugins/Online/EOSShared/Source/EOSShared/Public/IEOSSDKManager.h
--- a/Engine/Plugins/Online/EOSShared/Source/EOSShared/Public/IEOSSDKManager.h
+++ b/Engine/Plugins/Online/EOSShared/Source/EOSShared/Public/IEOSSDKManager.h
@@ -91,6 +91,7 @@ DECLARE_MULTICAST_DELEGATE_TwoParams(FEOSSDKManagerOnPreCreateNamedPlatform, con
 DECLARE_MULTICAST_DELEGATE_OneParam(FEOSSDKManagerOnPreCreatePlatform, EOS_Platform_Options& Options);
 DECLARE_MULTICAST_DELEGATE_OneParam(FEOSSDKManagerOnPlatformCreated, const IEOSPlatformHandlePtr& PlatformHandle);
 DECLARE_MULTICAST_DELEGATE_OneParam(FEOSSDKManagerOnPreReleasePlatform, const EOS_HPlatform& PlatformHandle);
+DECLARE_MULTICAST_DELEGATE_TwoParams(FEOSSDKManagerOnNetworkStatusChanged, const EOS_ENetworkStatus OldNetworkStatus, const EOS_ENetworkStatus NewNetworkStatus);
 DECLARE_DELEGATE_RetVal(FString, FEOSSDKManagerOnRequestRuntimeLibraryName);
 
 class IEOSSDKManager : public IModularFeature
@@ -160,6 +161,7 @@ public:
 	FEOSSDKManagerOnPlatformCreated OnPlatformCreated;
 	FEOSSDKManagerOnPreReleasePlatform OnPreReleasePlatform;
 	FEOSSDKManagerOnRequestRuntimeLibraryName OnRequestRuntimeLibraryName;
+	FEOSSDKManagerOnNetworkStatusChanged OnNetworkStatusChanged;
 };
 
 #endif // WITH_EOS_SDK
\ No newline at end of file
diff --git a/Engine/Plugins/Online/SocketSubsystemEOS/Source/SocketSubsystemEOS/Private/SocketEOS.cpp b/Engine/Plugins/Online/SocketSubsystemEOS/Source/SocketSubsystemEOS/Private/SocketEOS.cpp
--- a/Engine/Plugins/Online/SocketSubsystemEOS/Source/SocketSubsystemEOS/Private/SocketEOS.cpp
+++ b/Engine/Plugins/Online/SocketSubsystemEOS/Source/SocketSubsystemEOS/Private/SocketEOS.cpp
@@ -32,6 +32,7 @@ FSocketEOS::FSocketEOS(FSocketSubsystemEOS& InSocketSubsystem, const FString& In
 			PacketReliability = PacketReliabilityType;
 		}
 	}
+	IEOSSDKManager::Get()->OnNetworkStatusChanged.AddRaw(this, &FSocketEOS::OnNetworkStatusChanged);
 #endif
 }
 
@@ -53,6 +54,30 @@ bool FSocketEOS::Shutdown(ESocketShutdownMode Mode)
 	return false;
 }
 
+#if WITH_EOS_SDK
+void FSocketEOS::OnNetworkStatusChanged(EOS_ENetworkStatus OldNetworkStatus, EOS_ENetworkStatus NewNetworkStatus)
+{
+	// Only needed if we are listening on a socket, otherwise nothing needs to be done
+	if (bIsListening)
+	{
+		// Network status changes from an offline state to online - need to rebind P2P connection request notification 
+		if (OldNetworkStatus != EOS_ENetworkStatus::EOS_NS_Online && NewNetworkStatus == EOS_ENetworkStatus::EOS_NS_Online)
+		{
+			Listen(0);
+		}
+		// Network status changes from online to another state - clean up what was set on Listen()
+		else if (OldNetworkStatus == EOS_ENetworkStatus::EOS_NS_Online && NewNetworkStatus != EOS_ENetworkStatus::EOS_NS_Online)
+		{
+			// This is a special case where the P2P interface in the EOS SDK unbinds notifiers for us. We need to cleanup bound functions and notification Ids
+			ConnectNotifyId = EOS_INVALID_NOTIFICATIONID;
+			ClosedNotifyId = EOS_INVALID_NOTIFICATIONID;
+			ConnectNotifyCallback = nullptr;
+			ClosedNotifyCallback = nullptr;
+		}
+	}	
+}
+#endif
+
 bool FSocketEOS::Close()
 {
 	check(IsInGameThread() && "p2p does not support multithreading");
diff --git a/Engine/Plugins/Online/SocketSubsystemEOS/Source/SocketSubsystemEOS/Public/SocketEOS.h b/Engine/Plugins/Online/SocketSubsystemEOS/Source/SocketSubsystemEOS/Public/SocketEOS.h
--- a/Engine/Plugins/Online/SocketSubsystemEOS/Source/SocketSubsystemEOS/Public/SocketEOS.h
+++ b/Engine/Plugins/Online/SocketSubsystemEOS/Source/SocketSubsystemEOS/Public/SocketEOS.h
@@ -12,6 +12,7 @@ class FSocketSubsystemEOS;
 
 #if WITH_EOS_SDK
 	#include "eos_p2p_types.h"
+	#include "eos_types.h"
 #endif
 
 class FSocketEOS
@@ -70,6 +71,10 @@ public:
 	void RegisterClosedNotification();
 
 private:
+
+	/** Rebind connection request handler when network goes from offline to online, and cleanup objects when network goes from online to offline */
+	void OnNetworkStatusChanged(EOS_ENetworkStatus OldNetworkStatus, EOS_ENetworkStatus NewNetworkStatus);
+
 	/** Used to track our aliveness and make it possible to use the callback interface */
 	TSharedPtr<FCallbackBase> CallbackAliveTracker;
 
