Crash in FSocketSubsystemEOS in packaged builds

diff --git a/Engine/Plugins/Online/SocketSubsystemEOS/Source/SocketSubsystemEOS/Private/SocketSubsystemEOS.cpp b/Engine/Plugins/Online/SocketSubsystemEOS/Source/SocketSubsystemEOS/Private/SocketSubsystemEOS.cpp
--- a/Engine/Plugins/Online/SocketSubsystemEOS/Source/SocketSubsystemEOS/Private/SocketSubsystemEOS.cpp
+++ b/Engine/Plugins/Online/SocketSubsystemEOS/Source/SocketSubsystemEOS/Private/SocketSubsystemEOS.cpp
@@ -71,10 +71,7 @@ FSocketSubsystemEOS* FSocketSubsystemEOS::GetSocketSubsystemForWorld(UWorld* InW
 			if (NewWorld == InWorld)
 			{
 				SocketSubsystemEOSPerWorldMap.Add(InWorld, SocketSubsystem);
-
-				Result = &SocketSubsystem;
-
-				break;
+				return SocketSubsystem;
 			}
 		}
 	}
