Refactor NetDriverEOS to infer P2P channel from the netdriver definition name.
Removal of related logic handling passed-in URL port.

diff --git a/Engine/Plugins/Online/SocketSubsystemEOS/Source/SocketSubsystemEOS/Private/NetDriverEOS.cpp b/Engine/Plugins/Online/SocketSubsystemEOS/Source/SocketSubsystemEOS/Private/NetDriverEOS.cpp
--- a/Engine/Plugins/Online/SocketSubsystemEOS/Source/SocketSubsystemEOS/Private/NetDriverEOS.cpp
+++ b/Engine/Plugins/Online/SocketSubsystemEOS/Source/SocketSubsystemEOS/Private/NetDriverEOS.cpp
@@ -86,11 +86,13 @@ bool UNetDriverEOS::InitBase(bool bInitAsClient, FNetworkNotify* InNotify, const
 		return false;
 	}
 
+	FString NetDriverDefinitionStr = GetNetDriverDefinition().ToString();
+
 	// Store our local address and set our port
 	TSharedRef<FInternetAddrEOS> EOSLocalAddress = StaticCastSharedRef<FInternetAddrEOS>(LocalAddress);
-	EOSLocalAddress->SetChannel((uint8)URL.Port);
+	EOSLocalAddress->SetChannel(GetTypeHash(NetDriverDefinitionStr));
 	// Set our net driver name so we don't accept connections across net driver types
-	EOSLocalAddress->SetSocketName(GetNetDriverDefinition().ToString());
+	EOSLocalAddress->SetSocketName(NetDriverDefinitionStr);
 
 	static_cast<FSocketEOS*>(GetSocket())->SetLocalAddress(*EOSLocalAddress);
 
@@ -109,17 +111,6 @@ bool UNetDriverEOS::InitConnect(FNetworkNotify* InNotify, const FURL& ConnectURL
 		return Super::InitConnect(InNotify, ConnectURL, Error);
 	}
 
-	bool bIsValid = false;
-	TSharedRef<FInternetAddrEOS> RemoteHost = MakeShared<FInternetAddrEOS>();
-	RemoteHost->SetIp(*ConnectURL.Host, bIsValid);
-	RemoteHost->SetPort(ConnectURL.Port);
-	if (!bIsValid || ConnectURL.Port < 0)
-	{
-		Error = TEXT("Invalid remote address");
-		UE_LOG(LogTemp, Warning, TEXT("Invalid Remote Address. ConnectUrl = (%s)"), *ConnectURL.ToString());
-		return false;
-	}
-
 	UE_LOG(LogTemp, Verbose, TEXT("Connecting using EOSNetDriver. ConnectUrl = (%s)"), *ConnectURL.ToString());
 
 	if (!InitBase(true, InNotify, ConnectURL, false, Error))
@@ -127,9 +118,6 @@ bool UNetDriverEOS::InitConnect(FNetworkNotify* InNotify, const FURL& ConnectURL
 		return false;
 	}
 
-	// Set the address to what was parsed (us + remote)
-	LocalAddr = RemoteHost;
-
 	// Reference to our newly created socket
 	FSocket* CurSocket = GetSocket();
 
