Added VeryVerbose logging to FSocketEOS

diff --git a/Engine/Plugins/Online/SocketSubsystemEOS/Source/SocketSubsystemEOS/Private/SocketEOS.cpp b/Engine/Plugins/Online/SocketSubsystemEOS/Source/SocketSubsystemEOS/Private/SocketEOS.cpp
--- a/Engine/Plugins/Online/SocketSubsystemEOS/Source/SocketSubsystemEOS/Private/SocketEOS.cpp
+++ b/Engine/Plugins/Online/SocketSubsystemEOS/Source/SocketSubsystemEOS/Private/SocketEOS.cpp
@@ -109,9 +109,10 @@ bool FSocketEOS::Close()
 		Options.LocalUserId = SocketSubsystem.GetLocalUserId();
 		Options.SocketId = &SocketId;
 
-		EOS_EResult Result = EOS_P2P_CloseConnections(SocketSubsystem.GetP2PHandle(), &Options);
-
-		UE_LOG(LogSocketSubsystemEOS, Log, TEXT("Closing socket (%s) with result (%s)"), *LocalAddress.ToString(true), *LexToString(Result));
+		UE_LOG(LogSocketSubsystemEOS, VeryVerbose, TEXT("Calling EOS_P2P_CloseConnections SocketId=[%hs]"), SocketId.SocketName);
+		EOS_EResult CloseResult = EOS_P2P_CloseConnections(SocketSubsystem.GetP2PHandle(), &Options);
+		UE_CLOG(CloseResult == EOS_EResult::EOS_Success, LogSocketSubsystemEOS, Log, TEXT("EOS_P2P_CloseConnections SocketId=[%hs] Result=[%s]"), SocketId.SocketName, *LexToString(CloseResult));
+		UE_CLOG(CloseResult != EOS_EResult::EOS_Success, LogSocketSubsystemEOS, Error, TEXT("EOS_P2P_CloseConnections SocketId=[%hs] Result=[%s]"), SocketId.SocketName, *LexToString(CloseResult));
 
 		ClosedRemotes.Empty();
 	}
@@ -125,7 +126,7 @@ bool FSocketEOS::Bind(const FInternetAddr& Addr)
 
 	if (!Addr.IsValid())
 	{
-		UE_LOG(LogSocketSubsystemEOS, Warning, TEXT("Attempted to bind to invalid address. Address = (%s)"), *Addr.ToString(true));
+		UE_LOG(LogSocketSubsystemEOS, Warning, TEXT("Attempted to bind to invalid Address=[%s]"), *Addr.ToString(true));
 		SocketSubsystem.SetLastSocketError(ESocketErrors::SE_EADDRNOTAVAIL);
 		return false;
 	}
@@ -133,7 +134,7 @@ bool FSocketEOS::Bind(const FInternetAddr& Addr)
 	// If we have a remote user id, we're already bound
 	if (LocalAddress.GetRemoteUserId() != nullptr)
 	{
-		UE_LOG(LogSocketSubsystemEOS, Warning, TEXT("Attempted to bind a socket that was already bound. ExistingAddress = (%s) NewAddress = (%s)"), *LocalAddress.ToString(true), *Addr.ToString(true));
+		UE_LOG(LogSocketSubsystemEOS, Warning, TEXT("Attempted to bind a socket that was already bound ExistingAddress=[%s] NewAddress=[%s]"), *LocalAddress.ToString(true), *Addr.ToString(true));
 		SocketSubsystem.SetLastSocketError(ESocketErrors::SE_EADDRINUSE);
 		return false;
 	}
@@ -141,7 +142,7 @@ bool FSocketEOS::Bind(const FInternetAddr& Addr)
 	const FInternetAddrEOS& EOSAddr = static_cast<const FInternetAddrEOS&>(Addr);
 	if (!SocketSubsystem.BindChannel(EOSAddr))
 	{
-		UE_LOG(LogSocketSubsystemEOS, Warning, TEXT("Attempted to bind a socket to a port in use. NewAddress = (%s)"), *Addr.ToString(true));
+		UE_LOG(LogSocketSubsystemEOS, Warning, TEXT("Attempted to bind a socket to a port in use NewAddress=[%s]"), *Addr.ToString(true));
 		// BindPort sets our LastSocketError
 		return false;
 	}
@@ -154,7 +155,7 @@ bool FSocketEOS::Bind(const FInternetAddr& Addr)
 	LocalAddress = EOSAddr;
 	LocalAddress.SetLocalUserId(LocalUserId);
 
-	UE_LOG(LogSocketSubsystemEOS, Verbose, TEXT("Successfully bound socket to address (%s)"), *LocalAddress.ToString(true));
+	UE_LOG(LogSocketSubsystemEOS, Verbose, TEXT("Successfully bound socket to Address=[%s]"), *LocalAddress.ToString(true));
 	return true;
 }
 
@@ -171,7 +172,7 @@ bool FSocketEOS::Listen(int32)
 
 	if (!LocalAddress.IsValid())
 	{
-		UE_LOG(LogSocketSubsystemEOS, Warning, TEXT("Attempted to listen without a bound address. Address = (%s)"), *LocalAddress.ToString(true));
+		UE_LOG(LogSocketSubsystemEOS, Warning, TEXT("Attempted to listen without a bound Address=[%s]"), *LocalAddress.ToString(true));
 		SocketSubsystem.SetLastSocketError(ESocketErrors::SE_EADDRINUSE);
 		return false;
 	}
@@ -218,19 +219,14 @@ bool FSocketEOS::Listen(int32)
 			Options.LocalUserId = LocalAddress.GetLocalUserId();
 			Options.RemoteUserId = Info->RemoteUserId;
 			Options.SocketId = &SocketId;
+			UE_LOG(LogSocketSubsystemEOS, VeryVerbose, TEXT("Calling EOS_P2P_AcceptConnection SocketId=[%hs]"), SocketId.SocketName);
 			EOS_EResult AcceptResult = EOS_P2P_AcceptConnection(SocketSubsystem.GetP2PHandle(), &Options);
-			if (AcceptResult == EOS_EResult::EOS_Success)
-			{
-				UE_LOG(LogSocketSubsystemEOS, Verbose, TEXT("Accepting connection request from (%s) on socket (%s)"), *RemoteUser, UTF8_TO_TCHAR(Info->SocketId->SocketName));
-			}
-			else
-			{
-				UE_LOG(LogSocketSubsystemEOS, Error, TEXT("EOS_P2P_AcceptConnection from (%s) on socket (%s) failed with (%s)"), *RemoteUser, UTF8_TO_TCHAR(Info->SocketId->SocketName), *LexToString(AcceptResult));
-			}
+			UE_CLOG(AcceptResult == EOS_EResult::EOS_Success, LogSocketSubsystemEOS, Log, TEXT("EOS_P2P_AcceptConnection RemoteUser=[%s] SocketId=[%hs] Result=[%s]"), *RemoteUser, SocketId.SocketName, *LexToString(AcceptResult));
+			UE_CLOG(AcceptResult != EOS_EResult::EOS_Success, LogSocketSubsystemEOS, Error, TEXT("EOS_P2P_AcceptConnection RemoteUser=[%s] SocketId=[%hs] Result=[%s]"), *RemoteUser, SocketId.SocketName, *LexToString(AcceptResult));
 		}
 		else
 		{
-			UE_LOG(LogSocketSubsystemEOS, Warning, TEXT("Ignoring connection request from (%s) on socket (%s)"), *RemoteUser, UTF8_TO_TCHAR(Info->SocketId->SocketName));
+			UE_LOG(LogSocketSubsystemEOS, Warning, TEXT("Ignoring connection request from RemoteUser=[%s] SocketId=[%hs]"), *RemoteUser, Info->SocketId->SocketName);
 		}
 	};
 	ConnectNotifyId = EOS_P2P_AddNotifyPeerConnectionRequest(SocketSubsystem.GetP2PHandle(), &Options, ConnectNotifyCallback, ConnectNotifyCallback->GetCallbackPtr());
@@ -272,7 +268,7 @@ bool FSocketEOS::HasPendingData(uint32& PendingDataSize)
 	}
 	if (Result != EOS_EResult::EOS_Success)
 	{
-		UE_LOG(LogSocketSubsystemEOS, Warning, TEXT("Unable to check for data on address (%s) result code = (%s)"), *LocalAddress.ToString(true), *LexToString(Result));
+		UE_LOG(LogSocketSubsystemEOS, Warning, TEXT("Unable to check for data Address=[%s] Result=[%s]"), *LocalAddress.ToString(true), *LexToString(Result));
 
 		// @todo joeg - map EOS codes to UE4's
 		SocketSubsystem.SetLastSocketError(ESocketErrors::SE_EINVAL);
@@ -307,7 +303,7 @@ bool FSocketEOS::SendTo(const uint8* Data, int32 Count, int32& OutBytesSent, con
 
 	if (!Destination.IsValid())
 	{
-		UE_LOG(LogSocketSubsystemEOS, Warning, TEXT("Unable to send data, invalid destination address. DestinationAddress = (%s)"), *Destination.ToString(true));
+		UE_LOG(LogSocketSubsystemEOS, Warning, TEXT("Unable to send data, invalid DestinationAddress=[%s]"), *Destination.ToString(true));
 
 		SocketSubsystem.SetLastSocketError(ESocketErrors::SE_EADDRNOTAVAIL);
 		return false;
@@ -316,7 +312,7 @@ bool FSocketEOS::SendTo(const uint8* Data, int32 Count, int32& OutBytesSent, con
 #if WITH_EOS_SDK
 	if (Count > EOS_P2P_MAX_PACKET_SIZE)
 	{
-		UE_LOG(LogSocketSubsystemEOS, Warning, TEXT("Unable to send data, data over maximum size. Amount=[%d/%d] DestinationAddress = (%s)"), Count, EOS_P2P_MAX_PACKET_SIZE, *Destination.ToString(true));
+		UE_LOG(LogSocketSubsystemEOS, Warning, TEXT("Unable to send data, data over maximum size. Amount=[%d/%d] DestinationAddress=[%s]"), Count, EOS_P2P_MAX_PACKET_SIZE, *Destination.ToString(true));
 
 		SocketSubsystem.SetLastSocketError(ESocketErrors::SE_EMSGSIZE);
 		return false;
@@ -324,7 +320,7 @@ bool FSocketEOS::SendTo(const uint8* Data, int32 Count, int32& OutBytesSent, con
 
 	if (Count < 0)
 	{
-		UE_LOG(LogSocketSubsystemEOS, Warning, TEXT("Unable to send data, data invalid. Amount=[%d/%d] DestinationAddress = (%s)"), Count, EOS_P2P_MAX_PACKET_SIZE, *Destination.ToString(true));
+		UE_LOG(LogSocketSubsystemEOS, Warning, TEXT("Unable to send data, data invalid. Amount=[%d/%d] DestinationAddress=[%s]"), Count, EOS_P2P_MAX_PACKET_SIZE, *Destination.ToString(true));
 
 		SocketSubsystem.SetLastSocketError(ESocketErrors::SE_EINVAL);
 		return false;
@@ -333,7 +329,7 @@ bool FSocketEOS::SendTo(const uint8* Data, int32 Count, int32& OutBytesSent, con
 
 	if (Data == nullptr && Count != 0)
 	{
-		UE_LOG(LogSocketSubsystemEOS, Warning, TEXT("Unable to send data, data invalid. DestinationAddress = (%s)"), *Destination.ToString(true));
+		UE_LOG(LogSocketSubsystemEOS, Warning, TEXT("Unable to send data, data invalid. DestinationAddress=[%s]"), *Destination.ToString(true));
 
 		SocketSubsystem.SetLastSocketError(ESocketErrors::SE_EINVAL);
 		return false;
@@ -341,7 +337,7 @@ bool FSocketEOS::SendTo(const uint8* Data, int32 Count, int32& OutBytesSent, con
 
 	if (!LocalAddress.IsValid())
 	{
-		UE_LOG(LogSocketSubsystemEOS, Warning, TEXT("Unable to send data, socket was not initialized. DestinationAddress = (%s)"), *Destination.ToString(true));
+		UE_LOG(LogSocketSubsystemEOS, Warning, TEXT("Unable to send data, socket was not initialized. DestinationAddress=[%s]"), *Destination.ToString(true));
 
 		SocketSubsystem.SetLastSocketError(ESocketErrors::SE_NOTINITIALISED);
 		return false;
@@ -350,7 +346,7 @@ bool FSocketEOS::SendTo(const uint8* Data, int32 Count, int32& OutBytesSent, con
 	const FInternetAddrEOS& DestinationAddress = static_cast<const FInternetAddrEOS&>(Destination);
 	if (LocalAddress == DestinationAddress)
 	{
-		UE_LOG(LogSocketSubsystemEOS, Warning, TEXT("Unable to send data, unable to send data to ourselves. DestinationAddress = (%s)"), *Destination.ToString(true));
+		UE_LOG(LogSocketSubsystemEOS, Warning, TEXT("Unable to send data, unable to send data to ourselves. DestinationAddress=[%s]"), *Destination.ToString(true));
 
 		SocketSubsystem.SetLastSocketError(ESocketErrors::SE_ECONNREFUSED);
 		return false;
@@ -359,7 +355,7 @@ bool FSocketEOS::SendTo(const uint8* Data, int32 Count, int32& OutBytesSent, con
 	// Check for sending to an address we explicitly closed
 	if (WasClosed(DestinationAddress))
 	{
-		UE_LOG(LogSocketSubsystemEOS, Warning, TEXT("Unable to send data to closed connection. DestinationAddress = (%s)"), *Destination.ToString(true));
+		UE_LOG(LogSocketSubsystemEOS, Warning, TEXT("Unable to send data to closed connection. DestinationAddress=[%s]"), *Destination.ToString(true));
 
 		SocketSubsystem.SetLastSocketError(ESocketErrors::SE_ECONNREFUSED);
 		return false;
@@ -385,15 +381,18 @@ bool FSocketEOS::SendTo(const uint8* Data, int32 Count, int32& OutBytesSent, con
 	Options.Channel = DestinationAddress.GetChannel();
 	Options.DataLengthBytes = Count;
 	Options.Data = Data;
-	EOS_EResult Result = EOS_P2P_SendPacket(SocketSubsystem.GetP2PHandle(), &Options);
-	if (Result != EOS_EResult::EOS_Success)
+
+	UE_LOG(LogSocketSubsystemEOS, VeryVerbose, TEXT("Calling EOS_P2P_SendPacket DestinationAddress=[%s] SocketId=[%hs]"), *DestinationAddress.ToString(true), SocketId.SocketName);
+	EOS_EResult SendResult = EOS_P2P_SendPacket(SocketSubsystem.GetP2PHandle(), &Options);
+	if (SendResult != EOS_EResult::EOS_Success)
 	{
-		UE_LOG(LogSocketSubsystemEOS, Error, TEXT("Unable to send data to (%s) result code = (%s)"), *Destination.ToString(true), *LexToString(Result));
+		UE_LOG(LogSocketSubsystemEOS, Error, TEXT("EOS_P2P_SendPacket DestinationAddress=[%s] SocketId=[%hs] Result=[%s]"), *DestinationAddress.ToString(true), SocketId.SocketName, *LexToString(SendResult));
 
 		// @todo joeg - map EOS codes to UE4's
 		SocketSubsystem.SetLastSocketError(ESocketErrors::SE_EINVAL);
 		return false;
 	}
+	UE_LOG(LogSocketSubsystemEOS, VeryVerbose, TEXT("EOS_P2P_SendPacket DestinationAddress=[%s] SocketId=[%hs] Result=[%s]"), *DestinationAddress.ToString(true), SocketId.SocketName,*LexToString(SendResult));
 	OutBytesSent = Count;
 	return true;
 #else
@@ -416,7 +415,7 @@ bool FSocketEOS::RecvFrom(uint8* Data, int32 BufferSize, int32& BytesRead, FInte
 
 	if (BufferSize < 0)
 	{
-		UE_LOG(LogSocketSubsystemEOS, Error, TEXT("Unable to receive data, receiving buffer was invalid. BufferSize = (%d)"), BufferSize);
+		UE_LOG(LogSocketSubsystemEOS, Error, TEXT("Unable to receive data, receiving buffer was invalid. BufferSize=[%d]"), BufferSize);
 
 		SocketSubsystem.SetLastSocketError(ESocketErrors::SE_EINVAL);
 		return false;
@@ -425,7 +424,7 @@ bool FSocketEOS::RecvFrom(uint8* Data, int32 BufferSize, int32& BytesRead, FInte
 	if (Flags != ESocketReceiveFlags::None)
 	{
 		// We do not support peaking / blocking until a packet comes
-		UE_LOG(LogSocketSubsystemEOS, Error, TEXT("Socket receive flags (%d) are not supported"), int32(Flags));
+		UE_LOG(LogSocketSubsystemEOS, Error, TEXT("Socket receive Flags=[%d] are not supported"), int32(Flags));
 
 		SocketSubsystem.SetLastSocketError(ESocketErrors::SE_EOPNOTSUPP);
 		return false;
@@ -443,16 +442,17 @@ bool FSocketEOS::RecvFrom(uint8* Data, int32 BufferSize, int32& BytesRead, FInte
 	EOS_ProductUserId RemoteUserId = nullptr;
 	EOS_P2P_SocketId SocketId;
 	
-	EOS_EResult Result = EOS_P2P_ReceivePacket(SocketSubsystem.GetP2PHandle(), &Options, &RemoteUserId, &SocketId, &Channel, Data, (uint32*)&BytesRead);
-	if (Result == EOS_EResult::EOS_NotFound)
+	UE_LOG(LogSocketSubsystemEOS, VeryVerbose, TEXT("Calling EOS_P2P_ReceivePacket RequestedChannel=[%d]"), Options.RequestedChannel);
+	EOS_EResult ReceiveResult = EOS_P2P_ReceivePacket(SocketSubsystem.GetP2PHandle(), &Options, &RemoteUserId, &SocketId, &Channel, Data, (uint32*)&BytesRead);
+	if (ReceiveResult == EOS_EResult::EOS_NotFound)
 	{
 		// No data to read
 		SocketSubsystem.SetLastSocketError(ESocketErrors::SE_EWOULDBLOCK);
 		return false;
 	}
-	else if (Result != EOS_EResult::EOS_Success)
+	else if (ReceiveResult != EOS_EResult::EOS_Success)
 	{
-		UE_LOG(LogSocketSubsystemEOS, Error, TEXT("Unable to receive data result code = (%s)"), *LexToString(Result));
+		UE_LOG(LogSocketSubsystemEOS, Error, TEXT("EOS_P2P_ReceivePacket RequestedChannel=[%d] Result=[%s]"), Options.RequestedChannel, *LexToString(ReceiveResult));
 
 		// @todo joeg - map EOS codes to UE4's
 		SocketSubsystem.SetLastSocketError(ESocketErrors::SE_EINVAL);
@@ -465,6 +465,8 @@ bool FSocketEOS::RecvFrom(uint8* Data, int32 BufferSize, int32& BytesRead, FInte
 	SourceAddress.SetSocketName(SocketId.SocketName);
 	SourceAddress.SetChannel(Channel);
 
+	UE_LOG(LogSocketSubsystemEOS, VeryVerbose, TEXT("EOS_P2P_ReceivePacket RemoteAddress=[%s] SocketId=[%hs] Result=[%s]"), *SourceAddress.ToString(true), SocketId.SocketName, *LexToString(ReceiveResult));
+
 	return true;
 #else
 	return false;
@@ -598,7 +600,7 @@ bool FSocketEOS::Close(const FInternetAddrEOS& RemoteAddress)
 
 	if (!RemoteAddress.IsValid())
 	{
-		UE_LOG(LogSocketSubsystemEOS, Error, TEXT("Unable to close socket with remote address as it is invalid RemoteAddress = (%s)"), *RemoteAddress.ToString(true));
+		UE_LOG(LogSocketSubsystemEOS, Error, TEXT("Unable to close socket with RemoteAddress=[%s] as it is invalid"), *RemoteAddress.ToString(true));
 		SocketSubsystem.SetLastSocketError(ESocketErrors::SE_EADDRNOTAVAIL);
 		return false;
 	}
@@ -619,15 +621,18 @@ bool FSocketEOS::Close(const FInternetAddrEOS& RemoteAddress)
 	Options.RemoteUserId = RemoteAddress.GetRemoteUserId();
 	Options.SocketId = &SocketId;
 
-	EOS_EResult Result = EOS_P2P_CloseConnection(SocketSubsystem.GetP2PHandle(), &Options);
-	if (Result != EOS_EResult::EOS_Success)
+	UE_LOG(LogSocketSubsystemEOS, VeryVerbose, TEXT("Calling EOS_P2P_CloseConnection RemoteAddress=[%s] SocketId=[%hs])"), *RemoteAddress.ToString(true), SocketId.SocketName);
+	EOS_EResult CloseResult = EOS_P2P_CloseConnection(SocketSubsystem.GetP2PHandle(), &Options);
+	if (CloseResult != EOS_EResult::EOS_Success)
 	{
-		UE_LOG(LogSocketSubsystemEOS, Error, TEXT("Unable to close socket with remote address RemoteAddress (%s) due to error (%s)"), *RemoteAddress.ToString(true), *LexToString(Result));
+		UE_LOG(LogSocketSubsystemEOS, Error, TEXT("EOS_P2P_CloseConnection RemoteAddress=[%s] SocketId=[%hs] Result=[%s]"), *RemoteAddress.ToString(true), SocketId.SocketName, *LexToString(CloseResult));
 
 		// @todo joeg - map EOS codes to UE4's
 		SocketSubsystem.SetLastSocketError(ESocketErrors::SE_EINVAL);
 		return false;
 	}
+
+	UE_LOG(LogSocketSubsystemEOS, Log, TEXT("EOS_P2P_CloseConnection RemoteAddress=[%s] SocketId=[%hs] Result=[%s]"), *RemoteAddress.ToString(true), SocketId.SocketName, *LexToString(CloseResult));
 	return true;
 #else
 	return false;
