- Change some logs from using LogTemp to LogSocketSubsystemEOS

diff --git a/Engine/Plugins/Online/SocketSubsystemEOS/Source/SocketSubsystemEOS/Private/NetDriverEOS.cpp b/Engine/Plugins/Online/SocketSubsystemEOS/Source/SocketSubsystemEOS/Private/NetDriverEOS.cpp
--- a/Engine/Plugins/Online/SocketSubsystemEOS/Source/SocketSubsystemEOS/Private/NetDriverEOS.cpp
+++ b/Engine/Plugins/Online/SocketSubsystemEOS/Source/SocketSubsystemEOS/Private/NetDriverEOS.cpp
@@ -16,12 +16,12 @@ UNetDriverEOS::UNetDriverEOS(const FObjectInitializer& ObjectInitializer)
 	// We check if there is config for bIsUsingP2PSockets for any of the old locations of the plugin
 	if (GConfig->GetBool(TEXT("/Script/OnlineSubsystemEOS.NetDriverEOS"), TEXT("bIsUsingP2PSockets"), bUnused, GEngineIni))
 	{
-		UE_LOG(LogTemp, Warning, TEXT("bIsUsingP2PSockets is deprecated, please remove any related config values"));
+		UE_LOG(LogSocketSubsystemEOS, Warning, TEXT("bIsUsingP2PSockets is deprecated, please remove any related config values"));
 	}
 
 	if (GConfig->GetBool(TEXT("/Script/SocketSubsystemEOS.NetDriverEOSBase"), TEXT("bIsUsingP2PSockets"), bUnused, GEngineIni))
 	{
-		UE_LOG(LogTemp, Warning, TEXT("bIsUsingP2PSockets is deprecated, please remove any related config values"));
+		UE_LOG(LogSocketSubsystemEOS, Warning, TEXT("bIsUsingP2PSockets is deprecated, please remove any related config values"));
 	}
 }
 
@@ -45,20 +45,20 @@ bool UNetDriverEOS::InitBase(bool bInitAsClient, FNetworkNotify* InNotify, const
 {
 	if (bIsPassthrough)
 	{
-		UE_LOG(LogTemp, Verbose, TEXT("Running as pass-through"));
+		UE_LOG(LogSocketSubsystemEOS, Verbose, TEXT("Running as pass-through"));
 		return Super::InitBase(bInitAsClient, InNotify, URL, bReuseAddressAndPort, Error);
 	}
 
 	if (!UNetDriver::InitBase(bInitAsClient, InNotify, URL, bReuseAddressAndPort, Error))
 	{
-		UE_LOG(LogTemp, Warning, TEXT("Failed to init driver base"));
+		UE_LOG(LogSocketSubsystemEOS, Warning, TEXT("Failed to init driver base"));
 		return false;
 	}
 
 	FSocketSubsystemEOS* const SocketSubsystem = static_cast<FSocketSubsystemEOS*>(GetSocketSubsystem());
 	if (!SocketSubsystem)
 	{
-		UE_LOG(LogTemp, Warning, TEXT("Could not get socket subsystem"));
+		UE_LOG(LogSocketSubsystemEOS, Warning, TEXT("Could not get socket subsystem"));
 		return false;
 	}
 
@@ -71,7 +71,7 @@ bool UNetDriverEOS::InitBase(bool bInitAsClient, FNetworkNotify* InNotify, const
 	{
 		// Not logged in?
 		Error = TEXT("Could not bind local address");
-		UE_LOG(LogTemp, Warning, TEXT("Could not bind local address"));
+		UE_LOG(LogSocketSubsystemEOS, Warning, TEXT("Could not bind local address"));
 		return false;
 	}
 
@@ -82,7 +82,7 @@ bool UNetDriverEOS::InitBase(bool bInitAsClient, FNetworkNotify* InNotify, const
 
 	if (GetSocket() == nullptr)
 	{
-		UE_LOG(LogTemp, Warning, TEXT("Could not create socket"));
+		UE_LOG(LogSocketSubsystemEOS, Warning, TEXT("Could not create socket"));
 		return false;
 	}
 
@@ -105,13 +105,13 @@ bool UNetDriverEOS::InitConnect(FNetworkNotify* InNotify, const FURL& ConnectURL
 {
 	if (!IsAvailable() || !ConnectURL.Host.StartsWith(EOS_CONNECTION_URL_PREFIX, ESearchCase::IgnoreCase))
 	{
-		UE_LOG(LogTemp, Verbose, TEXT("Connecting using IPNetDriver passthrough. ConnectUrl = (%s)"), *ConnectURL.ToString());
+		UE_LOG(LogSocketSubsystemEOS, Verbose, TEXT("Connecting using IPNetDriver passthrough. ConnectUrl = (%s)"), *ConnectURL.ToString());
 
 		bIsPassthrough = true;
 		return Super::InitConnect(InNotify, ConnectURL, Error);
 	}
 
-	UE_LOG(LogTemp, Verbose, TEXT("Connecting using EOSNetDriver. ConnectUrl = (%s)"), *ConnectURL.ToString());
+	UE_LOG(LogSocketSubsystemEOS, Verbose, TEXT("Connecting using EOSNetDriver. ConnectUrl = (%s)"), *ConnectURL.ToString());
 
 	if (!InitBase(true, InNotify, ConnectURL, false, Error))
 	{
@@ -128,7 +128,7 @@ bool UNetDriverEOS::InitConnect(FNetworkNotify* InNotify, const FURL& ConnectURL
 	{
 		// Failure
 		Error = TEXT("Could not bind local port");
-		UE_LOG(LogTemp, Warning, TEXT("Could not bind local port in %d attempts"), MaxPortCountToTry);
+		UE_LOG(LogSocketSubsystemEOS, Warning, TEXT("Could not bind local port in %d attempts"), MaxPortCountToTry);
 		return false;
 	}
 
@@ -149,13 +149,13 @@ bool UNetDriverEOS::InitListen(FNetworkNotify* InNotify, FURL& LocalURL, bool bR
 {
 	if (!IsAvailable() || LocalURL.HasOption(TEXT("bIsLanMatch")) || LocalURL.HasOption(TEXT("bUseIPSockets")))
 	{
-		UE_LOG(LogTemp, Verbose, TEXT("Init as IPNetDriver listen server. LocalURL = (%s)"), *LocalURL.ToString());
+		UE_LOG(LogSocketSubsystemEOS, Verbose, TEXT("Init as IPNetDriver listen server. LocalURL = (%s)"), *LocalURL.ToString());
 
 		bIsPassthrough = true;
 		return Super::InitListen(InNotify, LocalURL, bReuseAddressAndPort, Error);
 	}
 
-	UE_LOG(LogTemp, Verbose, TEXT("Init as EOSNetDriver listen server. LocalURL = (%s)"), *LocalURL.ToString());
+	UE_LOG(LogSocketSubsystemEOS, Verbose, TEXT("Init as EOSNetDriver listen server. LocalURL = (%s)"), *LocalURL.ToString());
 
 	if (!InitBase(false, InNotify, LocalURL, bReuseAddressAndPort, Error))
 	{
@@ -167,13 +167,13 @@ bool UNetDriverEOS::InitListen(FNetworkNotify* InNotify, FURL& LocalURL, bool bR
 	if (!CurSocket->Listen(0))
 	{
 		Error = TEXT("Could not listen");
-		UE_LOG(LogTemp, Warning, TEXT("Could not listen on socket"));
+		UE_LOG(LogSocketSubsystemEOS, Warning, TEXT("Could not listen on socket"));
 		return false;
 	}
 
 	InitConnectionlessHandler();
 
-	UE_LOG(LogTemp, Verbose, TEXT("Initialized as an EOSP2P listen server"));
+	UE_LOG(LogSocketSubsystemEOS, Verbose, TEXT("Initialized as an EOSP2P listen server"));
 	return true;
 }
 
@@ -193,7 +193,7 @@ ISocketSubsystem* UNetDriverEOS::GetSocketSubsystem()
 
 void UNetDriverEOS::Shutdown()
 {
-	UE_LOG(LogTemp, Verbose, TEXT("Shutting down NetDriver"));
+	UE_LOG(LogSocketSubsystemEOS, Verbose, TEXT("Shutting down NetDriver"));
 
 	Super::Shutdown();
 
