Reduced members of FInternetAddEOS and refactored usage accordingly.

diff --git a/Engine/Plugins/Online/OnlineServicesEOSGS/Source/Private/Online/OnlineServicesEOSGS.cpp b/Engine/Plugins/Online/OnlineServicesEOSGS/Source/Private/Online/OnlineServicesEOSGS.cpp
--- a/Engine/Plugins/Online/OnlineServicesEOSGS/Source/Private/Online/OnlineServicesEOSGS.cpp
+++ b/Engine/Plugins/Online/OnlineServicesEOSGS/Source/Private/Online/OnlineServicesEOSGS.cpp
@@ -110,9 +110,8 @@ TOnlineResult<FGetResolvedConnectString> FOnlineServicesEOSGS::GetResolvedConnec
 				if (Lobby->LobbyId == Params.LobbyId)
 				{
 #if WITH_ENGINE
-					//It should look like this: "EOS:0002aeeb5b2d4388a3752dd6d31222ec:GameNetDriver:97"
-					FString NetDriverName = GetDefault<UNetDriverEOS>()->NetDriverName.ToString();
-					FInternetAddrEOS TempAddr(GetProductUserIdChecked(Lobby->OwnerAccountId), NetDriverName, GetTypeHash(NetDriverName));
+					//It should look like this: "EOS:0002aeeb5b2d4388a3752dd6d31222ec"
+					FInternetAddrEOS TempAddr(GetProductUserIdChecked(Lobby->OwnerAccountId));
 					return TOnlineResult<FGetResolvedConnectString>({ TempAddr.ToString(true) });
 #else
 					return TOnlineResult<FGetResolvedConnectString>(Errors::NotImplemented());
@@ -139,9 +138,8 @@ TOnlineResult<FGetResolvedConnectString> FOnlineServicesEOSGS::GetResolvedConnec
 #if WITH_ENGINE
 			TSharedRef<const ISession> Session = Result.GetOkValue().Session;
 
-			//It should look like this: "EOS:0002aeeb5b2d4388a3752dd6d31222ec:GameNetDriver:97"
-			FString NetDriverName = GetDefault<UNetDriverEOS>()->NetDriverName.ToString();
-			FInternetAddrEOS TempAddr(GetProductUserIdChecked(Session->GetOwnerAccountId()), NetDriverName, GetTypeHash(NetDriverName));
+			//It should look like this: "EOS:0002aeeb5b2d4388a3752dd6d31222ec"
+			FInternetAddrEOS TempAddr(GetProductUserIdChecked(Session->GetOwnerAccountId()));
 			return TOnlineResult<FGetResolvedConnectString>({ TempAddr.ToString(true) });
 #else
 			return TOnlineResult<FGetResolvedConnectString>(Errors::NotImplemented());
diff --git a/Engine/Plugins/Online/OnlineSubsystemEOS/Source/OnlineSubsystemEOS/Private/OnlineSessionEOS.cpp b/Engine/Plugins/Online/OnlineSubsystemEOS/Source/OnlineSubsystemEOS/Private/OnlineSessionEOS.cpp
--- a/Engine/Plugins/Online/OnlineSubsystemEOS/Source/OnlineSubsystemEOS/Private/OnlineSessionEOS.cpp
+++ b/Engine/Plugins/Online/OnlineSubsystemEOS/Source/OnlineSubsystemEOS/Private/OnlineSessionEOS.cpp
@@ -1606,8 +1606,7 @@ uint32 FOnlineSessionEOS::CreateEOSSession(int32 HostingPlayerNum, FNamedOnlineS
 	if (!IsRunningDedicatedServer() && IsNetDriverEOS(NAME_GameNetDriver))
 	{
 		// Because some platforms remap ports, we will use the ID of the name of the net driver to be our port instead
-		const FString NetDriverName = GetDefault<UNetDriverEOS>()->NetDriverName.ToString();
-		FInternetAddrEOS TempAddr(Options.LocalUserId, NetDriverName, GetTypeHash(NetDriverName));
+		FInternetAddrEOS TempAddr(Options.LocalUserId);
 		HostAddr = TempAddr.ToString(true);
 	}
 	else
@@ -3247,14 +3246,7 @@ static bool GetConnectStringFromSessionInfo(TSharedPtr<FOnlineSessionInfoEOS>& S
 		return false;
 	}
 
-	if (PortOverride != 0)
-	{
-		ConnectInfo = FString::Printf(TEXT("[%s]:%d"), *SessionInfo->HostAddr->ToString(false), PortOverride);
-	}
-	else
-	{
-		ConnectInfo = SessionInfo->HostAddr->ToString(true);
-	}
+	ConnectInfo = FString::Printf(TEXT("[%s]:%d"), *SessionInfo->HostAddr->ToString(false), PortOverride != 0 ? PortOverride : SessionInfo->HostAddr->GetPort());
 
 	if (SocketNameOverride != NAME_None)
 	{
@@ -4159,7 +4151,7 @@ uint32 FOnlineSessionEOS::CreateLobbySession(int32 HostingPlayerNum, FNamedOnlin
 				TSharedPtr<FOnlineSessionInfoEOS> SessionInfo = MakeShared<FOnlineSessionInfoEOS>(FOnlineSessionInfoEOS::Create(FUniqueNetIdEOSLobby::Create(UTF8_TO_TCHAR(Data->LobbyId))));
 				
 				const FName NetDriverName = GetDefault<UNetDriverEOS>()->NetDriverName;
-				SessionInfo->HostAddr = MakeShared<FInternetAddrEOS>(LocalProductUserId, NetDriverName.ToString(), FURL::UrlConfig.DefaultPort);
+				SessionInfo->HostAddr = MakeShared<FInternetAddrEOS>(LocalProductUserId);
 
 				Session->SessionInfo = SessionInfo;
 
@@ -4839,7 +4831,7 @@ void FOnlineSessionEOS::AddLobbySearchResult(const TSharedRef<FLobbyDetailsEOS>&
 		TSharedPtr<FOnlineSessionInfoEOS> SessionInfo = MakeShared<FOnlineSessionInfoEOS>(FOnlineSessionInfoEOS::Create(FUniqueNetIdEOSLobby::Create(UTF8_TO_TCHAR(LobbyDetailsInfo->LobbyId)), LobbyDetails));
 
 		const FName NetDriverName = GetDefault<UNetDriverEOS>()->NetDriverName;
-		SessionInfo->HostAddr = MakeShared<FInternetAddrEOS>(LobbyDetailsInfo->LobbyOwnerUserId, NetDriverName.ToString(), FURL::UrlConfig.DefaultPort);
+		SessionInfo->HostAddr = MakeShared<FInternetAddrEOS>(LobbyDetailsInfo->LobbyOwnerUserId);
 
 		SearchResult.Session.SessionInfo = SessionInfo;
 
diff --git a/Engine/Plugins/Online/SocketSubsystemEOS/Source/SocketSubsystemEOS/Private/InternetAddrEOS.cpp b/Engine/Plugins/Online/SocketSubsystemEOS/Source/SocketSubsystemEOS/Private/InternetAddrEOS.cpp
--- a/Engine/Plugins/Online/SocketSubsystemEOS/Source/SocketSubsystemEOS/Private/InternetAddrEOS.cpp
+++ b/Engine/Plugins/Online/SocketSubsystemEOS/Source/SocketSubsystemEOS/Private/InternetAddrEOS.cpp
@@ -11,32 +11,36 @@ inline uint8 PortToChannel(int32 InPort)
 }
 
 FInternetAddrEOS::FInternetAddrEOS()
-	: LocalUserId(nullptr)
-	, RemoteUserId(nullptr)
-	, Channel(0)
+	: ProductUserId(nullptr)
 {
-	SocketName[0] = '\0';
+
+}
+
+FInternetAddrEOS::FInternetAddrEOS(const FString& InProductUserId, const FString& InSocketName, const int32 InChannel)
+	: FInternetAddrEOS(InProductUserId)
+{
+
 }
 
-FInternetAddrEOS::FInternetAddrEOS(const FString& InRemoteUserId, const FString& InSocketName, const int32 InChannel)
-	: LocalUserId(nullptr)
-	, RemoteUserId(nullptr)
+FInternetAddrEOS::FInternetAddrEOS(const FString& InProductUserId)
+	: ProductUserId(nullptr)
 {
 #if WITH_EOS_SDK
-	RemoteUserId = EOS_ProductUserId_FromString(TCHAR_TO_UTF8(*InRemoteUserId));
+	ProductUserId = EOS_ProductUserId_FromString(TCHAR_TO_UTF8(*InProductUserId));
 #endif
-	FCStringAnsi::Strcpy(SocketName, TCHAR_TO_UTF8(*InSocketName));
-	Channel = PortToChannel(InChannel);
 }
 
 #if WITH_EOS_SDK
-FInternetAddrEOS::FInternetAddrEOS(const EOS_ProductUserId InRemoteUserId, const FString& InSocketName, const int32 InChannel)
-	: LocalUserId(nullptr)
-	, RemoteUserId(InRemoteUserId)
-	, Channel(InChannel)
+FInternetAddrEOS::FInternetAddrEOS(const EOS_ProductUserId InProductUserId, const FString& InSocketName, const int32 InChannel)
+	: FInternetAddrEOS(InProductUserId)
 {
-	FCStringAnsi::Strcpy(SocketName, TCHAR_TO_UTF8(*InSocketName));
-	Channel = PortToChannel(InChannel);
+
+}
+
+FInternetAddrEOS::FInternetAddrEOS(const EOS_ProductUserId InProductUserId)
+	: ProductUserId(InProductUserId)
+{
+
 }
 #endif
 
@@ -54,13 +58,11 @@ void FInternetAddrEOS::SetIp(const TCHAR* InAddr, bool& bIsValid)
 		return;
 	}
 
-	Channel = 0;
-
 	TArray<FString> UrlParts;
 	FString FullAddress = InAddr;
 	FullAddress.ParseIntoArray(UrlParts, EOS_URL_SEPARATOR, false);
-	// Expect URLs to look like "EOS:PUID:SocketName:Channel" and channel can be optional
-	if (UrlParts.Num() < 3 || UrlParts.Num() > 4)
+	// Expect URLs to look like "EOS:PUID"
+	if (UrlParts.Num() != 2)
 	{
 		return;
 	}
@@ -69,21 +71,13 @@ void FInternetAddrEOS::SetIp(const TCHAR* InAddr, bool& bIsValid)
 		return;
 	}
 #if WITH_EOS_SDK
-	RemoteUserId = EOS_ProductUserId_FromString(TCHAR_TO_UTF8(*UrlParts[1]));
-	if (EOS_ProductUserId_IsValid(RemoteUserId) == EOS_FALSE)
+	ProductUserId = EOS_ProductUserId_FromString(TCHAR_TO_UTF8(*UrlParts[1]));
+	if (EOS_ProductUserId_IsValid(ProductUserId) == EOS_FALSE)
 #endif
 	{
 		return;
 	}
-	if (UrlParts[2].Len() == 0)
-	{
-		return;
-	}
-	SetSocketName(UrlParts[2]);
-	if (UrlParts.Num() == 4)
-	{
-		Channel = FCString::Atoi(*UrlParts[3]);
-	}
+
 	bIsValid = true;
 }
 
@@ -96,12 +90,14 @@ void FInternetAddrEOS::GetIp(uint32& OutAddr) const
 
 void FInternetAddrEOS::SetPort(int32 InPort)
 {
-	Channel = PortToChannel(InPort);
+	UE_LOG(LogSocketSubsystemEOS, Warning, TEXT("Calls to FInternetAddrEOS::SetPort are not valid"));
 }
 
 int32 FInternetAddrEOS::GetPort() const
 {
-	return Channel;
+	UE_LOG(LogSocketSubsystemEOS, Warning, TEXT("Calls to FInternetAddrEOS::GetPort are not valid"));
+
+	return -1;
 }
 
 void FInternetAddrEOS::SetRawIp(const TArray<uint8>& RawAddr)
@@ -148,29 +144,24 @@ FString FInternetAddrEOS::ToString(bool bAppendPort) const
 	char PuidBuffer[64];
 	int32 BufferLen = 64;
 #if WITH_EOS_SDK
-	if (EOS_ProductUserId_ToString(RemoteUserId, PuidBuffer, &BufferLen) != EOS_EResult::EOS_Success)
+	if (EOS_ProductUserId_ToString(ProductUserId, PuidBuffer, &BufferLen) != EOS_EResult::EOS_Success)
 #endif
 	{
 		PuidBuffer[0] = '\0';
 	}
 
-	if (bAppendPort)
-	{
-		return FString::Printf(TEXT("%s%s%s%s%s%s%u"), EOS_CONNECTION_URL_PREFIX, EOS_URL_SEPARATOR, UTF8_TO_TCHAR(PuidBuffer), EOS_URL_SEPARATOR, UTF8_TO_TCHAR(SocketName), EOS_URL_SEPARATOR, Channel);
-	}
-
-	return FString::Printf(TEXT("%s%s%s%s%s"), EOS_CONNECTION_URL_PREFIX, EOS_URL_SEPARATOR, UTF8_TO_TCHAR(PuidBuffer), EOS_URL_SEPARATOR, UTF8_TO_TCHAR(SocketName));
+	return FString::Printf(TEXT("%s%s%s"), EOS_CONNECTION_URL_PREFIX, EOS_URL_SEPARATOR, UTF8_TO_TCHAR(PuidBuffer));
 }
 
 uint32 FInternetAddrEOS::GetTypeHash() const
 {
-	return HashCombine(HashCombine(HashCombine(GetTypeHashHelper((void*)LocalUserId), GetTypeHashHelper((void*)RemoteUserId)), GetTypeHashHelper(FAnsiStringView(SocketName, EOS_SOCKET_NAME_SIZE))), Channel);
+	return GetTypeHashHelper((void*)ProductUserId);
 }
 
 bool FInternetAddrEOS::IsValid() const
 {
 #if WITH_EOS_SDK
-	return (EOS_ProductUserId_IsValid(LocalUserId) == EOS_TRUE || EOS_ProductUserId_IsValid(RemoteUserId) == EOS_TRUE) && FCStringAnsi::Strlen(SocketName) > 0;
+	return (EOS_ProductUserId_IsValid(ProductUserId) == EOS_TRUE);
 #else
 	return false;
 #endif
diff --git a/Engine/Plugins/Online/SocketSubsystemEOS/Source/SocketSubsystemEOS/Private/NetDriverEOS.cpp b/Engine/Plugins/Online/SocketSubsystemEOS/Source/SocketSubsystemEOS/Private/NetDriverEOS.cpp
--- a/Engine/Plugins/Online/SocketSubsystemEOS/Source/SocketSubsystemEOS/Private/NetDriverEOS.cpp
+++ b/Engine/Plugins/Online/SocketSubsystemEOS/Source/SocketSubsystemEOS/Private/NetDriverEOS.cpp
@@ -79,7 +79,6 @@ bool UNetDriverEOS::InitBase(bool bInitAsClient, FNetworkNotify* InNotify, const
 
 	FUniqueSocket NewSocket = SocketSubsystem->CreateUniqueSocket(NAME_DGram, TEXT("UE4"), NAME_None);
 	TSharedPtr<FSocket> SharedSocket(NewSocket.Release(), FSocketDeleter(NewSocket.GetDeleter()));
-
 	SetSocketAndLocalAddress(SharedSocket);
 
 	if (GetSocket() == nullptr)
@@ -88,15 +87,14 @@ bool UNetDriverEOS::InitBase(bool bInitAsClient, FNetworkNotify* InNotify, const
 		return false;
 	}
 
-	FString NetDriverDefinitionStr = GetNetDriverDefinition().ToString();
+	TSharedPtr<FSocketEOS> SharedSocketEOS = StaticCastSharedPtr<FSocketEOS>(SharedSocket);
+	check(SharedSocketEOS.IsValid());
 
-	// Store our local address and set our port
-	TSharedRef<FInternetAddrEOS> EOSLocalAddress = StaticCastSharedRef<FInternetAddrEOS>(LocalAddress);
-	EOSLocalAddress->SetChannel(GetTypeHash(NetDriverDefinitionStr));
-	// Set our net driver name so we don't accept connections across net driver types
-	EOSLocalAddress->SetSocketName(NetDriverDefinitionStr);
+	SharedSocketEOS->SetSocketName(GetNetDriverDefinition().ToString());
 
-	static_cast<FSocketEOS*>(GetSocket())->SetLocalAddress(*EOSLocalAddress);
+	TSharedRef<FInternetAddrEOS> EOSLocalAddress = StaticCastSharedRef<FInternetAddrEOS>(LocalAddress);
+	
+	SharedSocketEOS->SetLocalAddress(*EOSLocalAddress);
 
 	LocalAddr = LocalAddress;
 
@@ -119,21 +117,9 @@ bool UNetDriverEOS::InitConnect(FNetworkNotify* InNotify, const FURL& ConnectURL
 	{
 		return false;
 	}
-
 	// Reference to our newly created socket
 	FSocket* CurSocket = GetSocket();
 
-	// Bind our local port
-	FSocketSubsystemEOS* const SocketSubsystem = static_cast<FSocketSubsystemEOS*>(GetSocketSubsystem());
-	check(SocketSubsystem);
-	if (!SocketSubsystem->BindNextPort(CurSocket, *LocalAddr, MaxPortCountToTry + 1, 1))
-	{
-		// Failure
-		Error = TEXT("Could not bind local port");
-		UE_LOG(LogSocketSubsystemEOS, Warning, TEXT("Could not bind local port in %d attempts"), MaxPortCountToTry);
-		return false;
-	}
-
 	// Create an unreal connection to the server
 	UNetConnectionEOS* Connection = NewObject<UNetConnectionEOS>(NetConnectionClass);
 	check(Connection);
diff --git a/Engine/Plugins/Online/SocketSubsystemEOS/Source/SocketSubsystemEOS/Private/SocketEOS.cpp b/Engine/Plugins/Online/SocketSubsystemEOS/Source/SocketSubsystemEOS/Private/SocketEOS.cpp
--- a/Engine/Plugins/Online/SocketSubsystemEOS/Source/SocketSubsystemEOS/Private/SocketEOS.cpp
+++ b/Engine/Plugins/Online/SocketSubsystemEOS/Source/SocketSubsystemEOS/Private/SocketEOS.cpp
@@ -40,11 +40,8 @@ FSocketEOS::~FSocketEOS()
 {
 	Close();
 
-	if (LocalAddress.IsValid())
-	{
-		SocketSubsystem.UnbindChannel(LocalAddress);
-		LocalAddress = FInternetAddrEOS();
-	}
+	SocketSubsystem.UnbindSocketName(SocketName);
+	LocalAddress = FInternetAddrEOS();
 	CallbackAliveTracker = nullptr;
 }
 
@@ -96,26 +93,22 @@ bool FSocketEOS::Close()
 	delete ClosedNotifyCallback;
 	ClosedNotifyCallback = nullptr;
 
-	if (LocalAddress.IsValid())
-	{
-		EOS_P2P_SocketId SocketId = { };
-		SocketId.ApiVersion = 1;
-		UE_EOS_CHECK_API_MISMATCH(EOS_P2P_SOCKETID_API_LATEST, 1);
-		FCStringAnsi::Strcpy(SocketId.SocketName, LocalAddress.GetSocketName());
-
-		EOS_P2P_CloseConnectionsOptions Options = { };
-		Options.ApiVersion = 1;
-		UE_EOS_CHECK_API_MISMATCH(EOS_P2P_CLOSECONNECTIONS_API_LATEST, 1);
-		Options.LocalUserId = SocketSubsystem.GetLocalUserId();
-		Options.SocketId = &SocketId;
-
-		UE_LOG(LogSocketSubsystemEOS, VeryVerbose, TEXT("Calling EOS_P2P_CloseConnections SocketId=[%hs]"), SocketId.SocketName);
-		EOS_EResult CloseResult = EOS_P2P_CloseConnections(SocketSubsystem.GetP2PHandle(), &Options);
-		UE_CLOG(CloseResult == EOS_EResult::EOS_Success, LogSocketSubsystemEOS, Log, TEXT("EOS_P2P_CloseConnections SocketId=[%hs] Result=[%s]"), SocketId.SocketName, *LexToString(CloseResult));
-		UE_CLOG(CloseResult != EOS_EResult::EOS_Success, LogSocketSubsystemEOS, Error, TEXT("EOS_P2P_CloseConnections SocketId=[%hs] Result=[%s]"), SocketId.SocketName, *LexToString(CloseResult));
-
-		ClosedRemotes.Empty();
-	}
+	EOS_P2P_SocketId SocketId = { };
+	SocketId.ApiVersion = 1;
+	UE_EOS_CHECK_API_MISMATCH(EOS_P2P_SOCKETID_API_LATEST, 1);
+	FCStringAnsi::Strcpy(SocketId.SocketName, TCHAR_TO_UTF8(*SocketName));
+
+	EOS_P2P_CloseConnectionsOptions Options = { };
+	Options.ApiVersion = 1;
+	UE_EOS_CHECK_API_MISMATCH(EOS_P2P_CLOSECONNECTIONS_API_LATEST, 1);
+	Options.LocalUserId = SocketSubsystem.GetLocalUserId();
+	Options.SocketId = &SocketId;
+
+	EOS_EResult Result = EOS_P2P_CloseConnections(SocketSubsystem.GetP2PHandle(), &Options);
+
+	UE_LOG(LogSocketSubsystemEOS, Log, TEXT("Closing socket (%s) with result (%s)"), *SocketName, *LexToString(Result));
+
+	ClosedRemotes.Empty();
 #endif
 	return true;
 }
@@ -126,36 +119,19 @@ bool FSocketEOS::Bind(const FInternetAddr& Addr)
 
 	if (!Addr.IsValid())
 	{
-		UE_LOG(LogSocketSubsystemEOS, Warning, TEXT("Attempted to bind to invalid Address=[%s]"), *Addr.ToString(true));
+		UE_LOG(LogSocketSubsystemEOS, Warning, TEXT("Attempted to bind to invalid address. Address = (%s)"), *Addr.ToString(true));
 		SocketSubsystem.SetLastSocketError(ESocketErrors::SE_EADDRNOTAVAIL);
 		return false;
 	}
 
-	// If we have a remote user id, we're already bound
-	if (LocalAddress.GetRemoteUserId() != nullptr)
-	{
-		UE_LOG(LogSocketSubsystemEOS, Warning, TEXT("Attempted to bind a socket that was already bound ExistingAddress=[%s] NewAddress=[%s]"), *LocalAddress.ToString(true), *Addr.ToString(true));
-		SocketSubsystem.SetLastSocketError(ESocketErrors::SE_EADDRINUSE);
-		return false;
-	}
-
-	const FInternetAddrEOS& EOSAddr = static_cast<const FInternetAddrEOS&>(Addr);
-	if (!SocketSubsystem.BindChannel(EOSAddr))
+	if (!SocketSubsystem.BindSocketName(SocketName))
 	{
-		UE_LOG(LogSocketSubsystemEOS, Warning, TEXT("Attempted to bind a socket to a port in use NewAddress=[%s]"), *Addr.ToString(true));
+		UE_LOG(LogSocketSubsystemEOS, Warning, TEXT("Attempted to bind a socket to a port in use. NewAddress = (%s)"), *Addr.ToString(true));
 		// BindPort sets our LastSocketError
 		return false;
 	}
 
-#if WITH_EOS_SDK
-	EOS_ProductUserId LocalUserId = LocalAddress.GetLocalUserId();
-#else
-	void* LocalUserId = LocalAddress.GetLocalUserId();
-#endif
-	LocalAddress = EOSAddr;
-	LocalAddress.SetLocalUserId(LocalUserId);
-
-	UE_LOG(LogSocketSubsystemEOS, Verbose, TEXT("Successfully bound socket to Address=[%s]"), *LocalAddress.ToString(true));
+	UE_LOG(LogSocketSubsystemEOS, Verbose, TEXT("Successfully bound socket (%s) to address (%s)"), *SocketName, *Addr.ToString(true));
 	return true;
 }
 
@@ -170,24 +146,17 @@ bool FSocketEOS::Listen(int32)
 {
 	check(IsInGameThread() && "p2p does not support multithreading");
 
-	if (!LocalAddress.IsValid())
-	{
-		UE_LOG(LogSocketSubsystemEOS, Warning, TEXT("Attempted to listen without a bound Address=[%s]"), *LocalAddress.ToString(true));
-		SocketSubsystem.SetLastSocketError(ESocketErrors::SE_EADDRINUSE);
-		return false;
-	}
-
 #if WITH_EOS_SDK
 	// Add listener for inbound connections
 	EOS_P2P_SocketId SocketId = { };
 	SocketId.ApiVersion = 1;
 	UE_EOS_CHECK_API_MISMATCH(EOS_P2P_SOCKETID_API_LATEST, 1);
-	FCStringAnsi::Strcpy(SocketId.SocketName, LocalAddress.GetSocketName());
+	FCStringAnsi::Strcpy(SocketId.SocketName, TCHAR_TO_UTF8(*SocketName));
 
 	EOS_P2P_AddNotifyPeerConnectionRequestOptions Options = { };
 	Options.ApiVersion = 1;
 	UE_EOS_CHECK_API_MISMATCH(EOS_P2P_ADDNOTIFYPEERCONNECTIONREQUEST_API_LATEST, 1);
-	Options.LocalUserId = LocalAddress.GetLocalUserId();
+	Options.LocalUserId = SocketSubsystem.GetLocalUserId();
 	Options.SocketId = &SocketId;
 
 	ConnectNotifyCallback = new FConnectNotifyCallback(CallbackAliveTracker);
@@ -201,11 +170,10 @@ bool FSocketEOS::Listen(int32)
 		}
 		FString RemoteUser(PuidBuffer);
 
-		if (Info->LocalUserId == LocalAddress.GetLocalUserId() && FCStringAnsi::Stricmp(Info->SocketId->SocketName, LocalAddress.GetSocketName()) == 0)
+		if (Info->LocalUserId == SocketSubsystem.GetLocalUserId() && FCStringAnsi::Stricmp(Info->SocketId->SocketName, TCHAR_TO_UTF8(*SocketName)) == 0)
 		{
 			// In case they disconnected and then reconnected, remove them from our closed list
-			FInternetAddrEOS RemoteAddress(Info->RemoteUserId, Info->SocketId->SocketName, LocalAddress.GetChannel());
-			RemoteAddress.SetLocalUserId(LocalAddress.GetLocalUserId());
+			FInternetAddrEOS RemoteAddress(Info->RemoteUserId);
 			ClosedRemotes.Remove(RemoteAddress);
 
 			EOS_P2P_SocketId SocketId = { };
@@ -216,7 +184,7 @@ bool FSocketEOS::Listen(int32)
 			EOS_P2P_AcceptConnectionOptions Options = { };
 			Options.ApiVersion = 1;
 			UE_EOS_CHECK_API_MISMATCH(EOS_P2P_ACCEPTCONNECTION_API_LATEST, 1);
-			Options.LocalUserId = LocalAddress.GetLocalUserId();
+			Options.LocalUserId = SocketSubsystem.GetLocalUserId();
 			Options.RemoteUserId = Info->RemoteUserId;
 			Options.SocketId = &SocketId;
 			UE_LOG(LogSocketSubsystemEOS, VeryVerbose, TEXT("Calling EOS_P2P_AcceptConnection SocketId=[%hs]"), SocketId.SocketName);
@@ -257,8 +225,8 @@ bool FSocketEOS::HasPendingData(uint32& PendingDataSize)
 	EOS_P2P_GetNextReceivedPacketSizeOptions Options = { };
 	Options.ApiVersion = 2;
 	UE_EOS_CHECK_API_MISMATCH(EOS_P2P_GETNEXTRECEIVEDPACKETSIZE_API_LATEST, 2);
-	Options.LocalUserId = LocalAddress.GetLocalUserId();
-	uint8 Channel = LocalAddress.GetChannel();
+	Options.LocalUserId = SocketSubsystem.GetLocalUserId();
+	uint8 Channel = GetTypeHash(SocketName);
 	Options.RequestedChannel = &Channel;
 
 	EOS_EResult Result = EOS_P2P_GetNextReceivedPacketSize(SocketSubsystem.GetP2PHandle(), &Options, &PendingDataSize);
@@ -268,7 +236,7 @@ bool FSocketEOS::HasPendingData(uint32& PendingDataSize)
 	}
 	if (Result != EOS_EResult::EOS_Success)
 	{
-		UE_LOG(LogSocketSubsystemEOS, Warning, TEXT("Unable to check for data Address=[%s] Result=[%s]"), *LocalAddress.ToString(true), *LexToString(Result));
+		UE_LOG(LogSocketSubsystemEOS, Warning, TEXT("EOS_P2P_GetNextReceivedPacketSize for socket [%s] (Channel [%d]) returned result code = [%s]"), *SocketName, Channel, *LexToString(Result));
 
 		// @todo joeg - map EOS codes to UE4's
 		SocketSubsystem.SetLastSocketError(ESocketErrors::SE_EINVAL);
@@ -335,22 +303,7 @@ bool FSocketEOS::SendTo(const uint8* Data, int32 Count, int32& OutBytesSent, con
 		return false;
 	}
 
-	if (!LocalAddress.IsValid())
-	{
-		UE_LOG(LogSocketSubsystemEOS, Warning, TEXT("Unable to send data, socket was not initialized. DestinationAddress=[%s]"), *Destination.ToString(true));
-
-		SocketSubsystem.SetLastSocketError(ESocketErrors::SE_NOTINITIALISED);
-		return false;
-	}
-
 	const FInternetAddrEOS& DestinationAddress = static_cast<const FInternetAddrEOS&>(Destination);
-	if (LocalAddress == DestinationAddress)
-	{
-		UE_LOG(LogSocketSubsystemEOS, Warning, TEXT("Unable to send data, unable to send data to ourselves. DestinationAddress=[%s]"), *Destination.ToString(true));
-
-		SocketSubsystem.SetLastSocketError(ESocketErrors::SE_ECONNREFUSED);
-		return false;
-	}
 
 	// Check for sending to an address we explicitly closed
 	if (WasClosed(DestinationAddress))
@@ -368,17 +321,17 @@ bool FSocketEOS::SendTo(const uint8* Data, int32 Count, int32& OutBytesSent, con
 	EOS_P2P_SocketId SocketId = { };
 	SocketId.ApiVersion = 1;
 	UE_EOS_CHECK_API_MISMATCH(EOS_P2P_SOCKETID_API_LATEST, 1);
-	FCStringAnsi::Strcpy(SocketId.SocketName, DestinationAddress.GetSocketName());
+	FCStringAnsi::Strcpy(SocketId.SocketName, TCHAR_TO_UTF8(*SocketName));
 
 	EOS_P2P_SendPacketOptions Options = { };
 	Options.ApiVersion = 3;
 	UE_EOS_CHECK_API_MISMATCH(EOS_P2P_SENDPACKET_API_LATEST, 3);
-	Options.LocalUserId = LocalAddress.GetLocalUserId();
-	Options.RemoteUserId = DestinationAddress.GetRemoteUserId();
+	Options.LocalUserId = SocketSubsystem.GetLocalUserId();
+	Options.RemoteUserId = DestinationAddress.GetProductUserId();
 	Options.SocketId = &SocketId;
 	Options.bAllowDelayedDelivery = EOS_TRUE;
 	Options.Reliability = PacketReliability;
-	Options.Channel = DestinationAddress.GetChannel();
+	Options.Channel = GetTypeHash(SocketName);
 	Options.DataLengthBytes = Count;
 	Options.Data = Data;
 
@@ -441,9 +394,9 @@ bool FSocketEOS::RecvFrom(uint8* Data, int32 BufferSize, int32& BytesRead, FInte
 	EOS_P2P_ReceivePacketOptions Options = { };
 	Options.ApiVersion = 2;
 	UE_EOS_CHECK_API_MISMATCH(EOS_P2P_RECEIVEPACKET_API_LATEST, 2);
-	Options.LocalUserId = LocalAddress.GetLocalUserId();
+	Options.LocalUserId = SocketSubsystem.GetLocalUserId();
 	Options.MaxDataSizeBytes = BufferSize;
-	uint8 Channel = LocalAddress.GetChannel();
+	uint8 Channel = GetTypeHash(SocketName);
 	Options.RequestedChannel = &Channel;
 
 	EOS_ProductUserId RemoteUserId = nullptr;
@@ -467,12 +420,7 @@ bool FSocketEOS::RecvFrom(uint8* Data, int32 BufferSize, int32& BytesRead, FInte
 	}
 
 	FInternetAddrEOS& SourceAddress = static_cast<FInternetAddrEOS&>(Source);
-	SourceAddress.SetLocalUserId(LocalAddress.GetLocalUserId());
-	SourceAddress.SetRemoteUserId(RemoteUserId);
-	SourceAddress.SetSocketName(SocketId.SocketName);
-	SourceAddress.SetChannel(Channel);
-
-	UE_LOG(LogSocketSubsystemEOS, VeryVerbose, TEXT("EOS_P2P_ReceivePacket RemoteAddress=[%s] SocketId=[%hs] Result=[%s]"), *SourceAddress.ToString(true), SocketId.SocketName, *LexToString(ReceiveResult));
+	SourceAddress.SetProductUserId(RemoteUserId);
 
 	return true;
 #else
@@ -593,7 +541,7 @@ bool FSocketEOS::SetReceiveBufferSize(int32 Size, int32& NewSize)
 
 int32 FSocketEOS::GetPortNo()
 {
-	return LocalAddress.GetChannel();
+	return GetTypeHash(SocketName);
 }
 
 void FSocketEOS::SetLocalAddress(const FInternetAddrEOS& InLocalAddress)
@@ -601,6 +549,11 @@ void FSocketEOS::SetLocalAddress(const FInternetAddrEOS& InLocalAddress)
 	LocalAddress = InLocalAddress;
 }
 
+void FSocketEOS::SetSocketName(const FString& InSocketName)
+{
+	SocketName = InSocketName;
+}
+
 bool FSocketEOS::Close(const FInternetAddrEOS& RemoteAddress)
 {
 	check(IsInGameThread() && "p2p does not support multithreading");
@@ -619,13 +572,13 @@ bool FSocketEOS::Close(const FInternetAddrEOS& RemoteAddress)
 	EOS_P2P_SocketId SocketId = { };
 	SocketId.ApiVersion = 1;
 	UE_EOS_CHECK_API_MISMATCH(EOS_P2P_SOCKETID_API_LATEST, 1);
-	FCStringAnsi::Strcpy(SocketId.SocketName, RemoteAddress.GetSocketName());
+	FCStringAnsi::Strcpy(SocketId.SocketName, TCHAR_TO_UTF8(*SocketName));
 
 	EOS_P2P_CloseConnectionOptions Options = { };
 	Options.ApiVersion = 1;
 	UE_EOS_CHECK_API_MISMATCH(EOS_P2P_CLOSECONNECTION_API_LATEST, 1);
-	Options.LocalUserId = LocalAddress.GetLocalUserId();
-	Options.RemoteUserId = RemoteAddress.GetRemoteUserId();
+	Options.LocalUserId = SocketSubsystem.GetLocalUserId();
+	Options.RemoteUserId = RemoteAddress.GetProductUserId();
 	Options.SocketId = &SocketId;
 
 	UE_LOG(LogSocketSubsystemEOS, VeryVerbose, TEXT("Calling EOS_P2P_CloseConnection RemoteAddress=[%s] SocketId=[%hs])"), *RemoteAddress.ToString(true), SocketId.SocketName);
@@ -658,20 +611,19 @@ void FSocketEOS::RegisterClosedNotification()
 	EOS_P2P_SocketId SocketId = { };
 	SocketId.ApiVersion = 1;
 	UE_EOS_CHECK_API_MISMATCH(EOS_P2P_SOCKETID_API_LATEST, 1);
-	FCStringAnsi::Strcpy(SocketId.SocketName, LocalAddress.GetSocketName());
+	FCStringAnsi::Strcpy(SocketId.SocketName, TCHAR_TO_UTF8(*SocketName));
 
 	EOS_P2P_AddNotifyPeerConnectionClosedOptions Options = { };
 	Options.ApiVersion = 1;
 	UE_EOS_CHECK_API_MISMATCH(EOS_P2P_ADDNOTIFYPEERCONNECTIONCLOSED_API_LATEST, 1);
-	Options.LocalUserId = LocalAddress.GetLocalUserId();
+	Options.LocalUserId = SocketSubsystem.GetLocalUserId();
 	Options.SocketId = &SocketId;
 
 	ClosedNotifyCallback = new FClosedNotifyCallback(CallbackAliveTracker);
 	ClosedNotifyCallback->CallbackLambda = [this](const EOS_P2P_OnRemoteConnectionClosedInfo* Info)
 	{
 		// Add this connection to the list of closed ones
-		FInternetAddrEOS RemoteAddress(Info->RemoteUserId, Info->SocketId->SocketName, LocalAddress.GetChannel());
-		RemoteAddress.SetLocalUserId(LocalAddress.GetLocalUserId());
+		FInternetAddrEOS RemoteAddress(Info->RemoteUserId);
 		ClosedRemotes.Add(RemoteAddress);
 	};
 	ClosedNotifyId = EOS_P2P_AddNotifyPeerConnectionClosed(SocketSubsystem.GetP2PHandle(), &Options, ClosedNotifyCallback, ClosedNotifyCallback->GetCallbackPtr());
diff --git a/Engine/Plugins/Online/SocketSubsystemEOS/Source/SocketSubsystemEOS/Private/SocketSubsystemEOS.cpp b/Engine/Plugins/Online/SocketSubsystemEOS/Source/SocketSubsystemEOS/Private/SocketSubsystemEOS.cpp
--- a/Engine/Plugins/Online/SocketSubsystemEOS/Source/SocketSubsystemEOS/Private/SocketSubsystemEOS.cpp
+++ b/Engine/Plugins/Online/SocketSubsystemEOS/Source/SocketSubsystemEOS/Private/SocketSubsystemEOS.cpp
@@ -247,24 +247,15 @@ TSharedRef<FInternetAddr> FSocketSubsystemEOS::GetLocalBindAddr(const UWorld* co
 
 #if WITH_EOS_SDK
 	EOS_ProductUserId LocalUserId = GetLocalUserId();
-	if (LocalUserId == nullptr)
+	if (LocalUserId != nullptr)
 	{
-		UE_LOG(LogSocketSubsystemEOS, Error, TEXT("No local user to send p2p packets with"));
-		return BoundAddr;
+		BoundAddr->SetProductUserId(LocalUserId);
 	}
-	BoundAddr->SetLocalUserId(LocalUserId);
-#else
-	return BoundAddr;
-#endif
-
-	FString SessionId = Utils->GetSessionId();
-
-	if (SessionId.IsEmpty())
+	else
 	{
-		SessionId = FName(NAME_GameSession).ToString();
+		UE_LOG(LogSocketSubsystemEOS, Error, TEXT("No local user set in LocalBindAddr"));
 	}
-
-	BoundAddr->SetSocketName(SessionId);
+#endif
 
 	return BoundAddr;
 }
@@ -279,59 +270,26 @@ void FSocketSubsystemEOS::SetLastSocketError(const ESocketErrors NewSocketError)
 	LastSocketError = NewSocketError;
 }
 
-bool FSocketSubsystemEOS::BindChannel(const FInternetAddrEOS& Address)
+bool FSocketSubsystemEOS::BindSocketName(const FString& SocketName)
 {
-	if (!Address.IsValid())
-	{
-		SetLastSocketError(ESocketErrors::SE_EINVAL);
-		return false;
-	}
-
-	const uint8 Channel = Address.GetChannel();
-
-	FChannelSet& ExistingBoundPorts = BoundAddresses.FindOrAdd(Address.GetSocketName());
-	if (ExistingBoundPorts.Contains(Channel))
+	if (BoundSocketNames.Contains(SocketName))
 	{
 		SetLastSocketError(ESocketErrors::SE_EADDRINUSE);
 		return false;
 	}
 
-	ExistingBoundPorts.Add(Channel);
+	BoundSocketNames.Add(SocketName);
 	return true;
 }
 
-bool FSocketSubsystemEOS::UnbindChannel(const FInternetAddrEOS& Address)
+bool FSocketSubsystemEOS::UnbindSocketName(const FString& SocketName)
 {
-	if (!Address.IsValid())
-	{
-		SetLastSocketError(ESocketErrors::SE_EINVAL);
-		return false;
-	}
-
-	const FString SocketName = Address.GetSocketName();
-	const uint8 Channel = Address.GetChannel();
-
-	// Find our sessions collection of ports
-	FChannelSet* ExistingBoundPorts = BoundAddresses.Find(SocketName);
-	if (!ExistingBoundPorts)
-	{
-		SetLastSocketError(ESocketErrors::SE_ENOTSOCK);
-		return false;
-	}
-
 	// Remove our port and check if we had it bound
-	if (ExistingBoundPorts->Remove(Channel) == 0)
+	if (BoundSocketNames.Remove(SocketName) == 0)
 	{
 		SetLastSocketError(ESocketErrors::SE_ENOTSOCK);
 		return false;
 	}
 
-	// Remove any empty sets
-	if (ExistingBoundPorts->Num() == 0)
-	{
-		BoundAddresses.Remove(SocketName);
-		ExistingBoundPorts = nullptr;
-	}
-
 	return true;
 }
diff --git a/Engine/Plugins/Online/SocketSubsystemEOS/Source/SocketSubsystemEOS/Public/InternetAddrEOS.h b/Engine/Plugins/Online/SocketSubsystemEOS/Source/SocketSubsystemEOS/Public/InternetAddrEOS.h
--- a/Engine/Plugins/Online/SocketSubsystemEOS/Source/SocketSubsystemEOS/Public/InternetAddrEOS.h
+++ b/Engine/Plugins/Online/SocketSubsystemEOS/Source/SocketSubsystemEOS/Public/InternetAddrEOS.h
@@ -21,9 +21,14 @@ class FInternetAddrEOS
 {
 public:
 	SOCKETSUBSYSTEMEOS_API FInternetAddrEOS();
+
+	UE_DEPRECATED(5.6, "InSocketName and InChannel have been removed from FInternetAddrEOS, please use the new constructor.")
 	SOCKETSUBSYSTEMEOS_API FInternetAddrEOS(const FString& InRemoteUserId, const FString& InSocketName, const int32 InChannel);
+	SOCKETSUBSYSTEMEOS_API FInternetAddrEOS(const FString& InRemoteUserId);
 #if WITH_EOS_SDK
+	UE_DEPRECATED(5.6, "InSocketName and InChannel have been removed from FInternetAddrEOS, please use the new constructor.")
 	SOCKETSUBSYSTEMEOS_API FInternetAddrEOS(const EOS_ProductUserId InRemoteUserId, const FString& InSocketName, const int32 InChannel);
+	SOCKETSUBSYSTEMEOS_API FInternetAddrEOS(const EOS_ProductUserId InRemoteUserId);
 #endif
 	virtual ~FInternetAddrEOS() = default;
 
@@ -46,21 +51,13 @@ public:
 
 	inline FInternetAddrEOS& operator=(const FInternetAddrEOS& Other)
 	{
-		LocalUserId = Other.LocalUserId;
-		RemoteUserId = Other.RemoteUserId;
-		FCStringAnsi::Strcpy(SocketName, Other.SocketName);
-		Channel = Other.Channel;
+		ProductUserId = Other.ProductUserId;
 		return *this;
 	}
 	
 	inline friend bool operator==(const FInternetAddrEOS& A, const FInternetAddrEOS& B)
 	{
-		return A.Channel == B.Channel
-#if WITH_EOS_SDK
-			&& A.LocalUserId == B.LocalUserId
-			&& A.RemoteUserId == B.RemoteUserId
-#endif
-			&& FCStringAnsi::Stricmp(A.GetSocketName(), B.GetSocketName()) == 0;
+		return A.ProductUserId == B.ProductUserId;
 	}
 
 	inline friend bool operator!=(const FInternetAddrEOS& A, const FInternetAddrEOS& B)
@@ -75,102 +72,37 @@ public:
 
 	friend bool operator<(const FInternetAddrEOS& Left, const FInternetAddrEOS& Right)
 	{
-#if WITH_EOS_SDK
-		if (Left.GetLocalUserId() < Right.GetLocalUserId())
-		{
-			return true;
-		}
-		else if (Left.GetRemoteUserId() < Right.GetRemoteUserId())
-		{
-			return true;
-		}
-		else
-#endif
-		if (FCStringAnsi::Stricmp(Left.GetSocketName(), Right.GetSocketName()) < 0)
-		{
-			return true;
-		}
-
-		return Left.GetChannel() < Right.GetChannel();
+		return Left.ProductUserId < Right.ProductUserId;
 	}
 
 #if WITH_EOS_SDK
-	EOS_ProductUserId GetLocalUserId() const
-	{
-		return LocalUserId;
-	}
-
-	void SetLocalUserId(EOS_ProductUserId InLocalUserId)
+	void SetProductUserId(EOS_ProductUserId InProductUserId)
 	{
-		LocalUserId = InLocalUserId;
+		ProductUserId = InProductUserId;
 	}
 
-	void SetRemoteUserId(EOS_ProductUserId InRemoteUserId)
+	EOS_ProductUserId GetProductUserId() const
 	{
-		RemoteUserId = InRemoteUserId;
-	}
-
-	EOS_ProductUserId GetRemoteUserId() const
-	{
-		return RemoteUserId;
+		return ProductUserId;
 	}
 #else
-	void* GetLocalUserId() const
-	{
-		return LocalUserId;
-	}
-
-	void SetLocalUserId(void* InLocalUserId)
+	void SetProductUserId(void* InProductUserId)
 	{
-		LocalUserId = InLocalUserId;
+		ProductUserId = InProductUserId;
 	}
 
-	void SetRemoteUserId(void* InRemoteUserId)
+	void* GetProductUserId() const
 	{
-		RemoteUserId = InRemoteUserId;
-	}
-
-	void* GetRemoteUserId() const
-	{
-		return RemoteUserId;
+		return ProductUserId;
 	}
 #endif
-	
-	const char* GetSocketName() const
-	{
-		return SocketName;
-	}
-
-	void SetSocketName(const FString& InSocketName)
-	{
-		FCStringAnsi::Strncpy(SocketName, TCHAR_TO_UTF8(*InSocketName), EOS_SOCKET_NAME_SIZE);
-	}
-
-	void SetSocketName(const char* InSocketName)
-	{
-		FCStringAnsi::Strncpy(SocketName, InSocketName, EOS_SOCKET_NAME_SIZE);
-	}
-
-	uint8 GetChannel() const
-	{
-		return Channel;
-	}
-
-	void SetChannel(uint8 InChannel)
-	{
-		Channel = InChannel;
-	}
 
 private:
 #if WITH_EOS_SDK
-	EOS_ProductUserId LocalUserId;
-	EOS_ProductUserId RemoteUserId;
+	EOS_ProductUserId ProductUserId;
 #else
-	void* LocalUserId;
-	void* RemoteUserId;
+	void* ProductUserId;
 #endif
-	char SocketName[EOS_SOCKET_NAME_SIZE];
-	uint8 Channel;
 
 	friend class SocketSubsystemEOS;
 };
diff --git a/Engine/Plugins/Online/SocketSubsystemEOS/Source/SocketSubsystemEOS/Public/SocketEOS.h b/Engine/Plugins/Online/SocketSubsystemEOS/Source/SocketSubsystemEOS/Public/SocketEOS.h
--- a/Engine/Plugins/Online/SocketSubsystemEOS/Source/SocketSubsystemEOS/Public/SocketEOS.h
+++ b/Engine/Plugins/Online/SocketSubsystemEOS/Source/SocketSubsystemEOS/Public/SocketEOS.h
@@ -60,6 +60,8 @@ public:
 
 	void SetLocalAddress(const FInternetAddrEOS& InLocalAddress);
 
+	void SetSocketName(const FString& InSocketName);
+
 	bool Close(const FInternetAddrEOS& RemoteAddress);
 
 	bool WasClosed(const FInternetAddrEOS& RemoteAddress)
@@ -83,6 +85,9 @@ private:
 
 	/** Our local address; session/port will be invalid when not bound */
 	FInternetAddrEOS LocalAddress;
+	
+	/** Name of the NetDriver that will use the socket */
+	FString SocketName;
 
 	/** Are we currently listening? */
 	bool bIsListening;
diff --git a/Engine/Plugins/Online/SocketSubsystemEOS/Source/SocketSubsystemEOS/Public/SocketSubsystemEOS.h b/Engine/Plugins/Online/SocketSubsystemEOS/Source/SocketSubsystemEOS/Public/SocketSubsystemEOS.h
--- a/Engine/Plugins/Online/SocketSubsystemEOS/Source/SocketSubsystemEOS/Public/SocketSubsystemEOS.h
+++ b/Engine/Plugins/Online/SocketSubsystemEOS/Source/SocketSubsystemEOS/Public/SocketSubsystemEOS.h
@@ -69,18 +69,16 @@ public:
 	/**
 	 * Bind our socket name & channel and ensure no other connections are using this combination
 	 *
-	 * @param Address The address to bind to
 	 * @return True if we bound successfully, false if not and LastSocketError will be set with a reason
 	 */
-	bool BindChannel(const FInternetAddrEOS& Address);
+	bool BindSocketName(const FString& SocketName);
 
 	/**
 	 * Unbind a previously bound socket name and channel.
 	 *
-	 * @param Address The Address to unbind
 	 * @return True if we released the port successfully, false if not and LastSocketError will be set with a reason
 	 */
-	bool UnbindChannel(const FInternetAddrEOS& Address);
+	bool UnbindSocketName(const FString& SocketName);
 
 	/**
 	 * Returns the FSocketSubsytemEOS instance corresponding with the UWorld* parameter.
@@ -112,7 +110,7 @@ private:
 	TArray<TUniquePtr<FSocketEOS>> TrackedSockets;
 
 	/** All bound socket addresses */
-	TMap<FString, FChannelSet> BoundAddresses;
+	TSet<FString> BoundSocketNames;
 
 	/** The last error we received */
 	ESocketErrors LastSocketError;
