Minor changes as followup to CL 44849953

diff --git a/Engine/Plugins/Online/OnlineSubsystemEOS/Source/OnlineSubsystemEOS/Private/OnlineSessionEOS.cpp b/Engine/Plugins/Online/OnlineSubsystemEOS/Source/OnlineSubsystemEOS/Private/OnlineSessionEOS.cpp
--- a/Engine/Plugins/Online/OnlineSubsystemEOS/Source/OnlineSubsystemEOS/Private/OnlineSessionEOS.cpp
+++ b/Engine/Plugins/Online/OnlineSubsystemEOS/Source/OnlineSubsystemEOS/Private/OnlineSessionEOS.cpp
@@ -1605,7 +1605,6 @@ uint32 FOnlineSessionEOS::CreateEOSSession(int32 HostingPlayerNum, FNamedOnlineS
 
 	if (!IsRunningDedicatedServer() && IsNetDriverEOS(NAME_GameNetDriver))
 	{
-		// Because some platforms remap ports, we will use the ID of the name of the net driver to be our port instead
 		FInternetAddrEOS TempAddr(Options.LocalUserId);
 		HostAddr = TempAddr.ToString(true);
 	}
@@ -3246,7 +3245,7 @@ static bool GetConnectStringFromSessionInfo(TSharedPtr<FOnlineSessionInfoEOS>& S
 		return false;
 	}
 
-	ConnectInfo = FString::Printf(TEXT("[%s]:%d"), *SessionInfo->HostAddr->ToString(false), PortOverride != 0 ? PortOverride : SessionInfo->HostAddr->GetPort());
+	ConnectInfo = FString::Printf(TEXT("[%s]:%d"), *SessionInfo->HostAddr->ToString(false), PortOverride);
 
 	if (SocketNameOverride != NAME_None)
 	{
diff --git a/Engine/Plugins/Online/SocketSubsystemEOS/Source/SocketSubsystemEOS/Public/SocketSubsystemEOS.h b/Engine/Plugins/Online/SocketSubsystemEOS/Source/SocketSubsystemEOS/Public/SocketSubsystemEOS.h
--- a/Engine/Plugins/Online/SocketSubsystemEOS/Source/SocketSubsystemEOS/Public/SocketSubsystemEOS.h
+++ b/Engine/Plugins/Online/SocketSubsystemEOS/Source/SocketSubsystemEOS/Public/SocketSubsystemEOS.h
@@ -67,14 +67,14 @@ public:
 #endif
 
 	/**
-	 * Bind our socket name & channel and ensure no other connections are using this combination
+	 * Bind our socket name and ensure no other connections are using this combination
 	 *
 	 * @return True if we bound successfully, false if not and LastSocketError will be set with a reason
 	 */
 	bool BindSocketName(const FString& SocketName);
 
 	/**
-	 * Unbind a previously bound socket name and channel.
+	 * Unbind a previously bound socket name.
 	 *
 	 * @return True if we released the port successfully, false if not and LastSocketError will be set with a reason
 	 */
@@ -109,7 +109,7 @@ private:
 	/** All sockets allocated by this subsystem */
 	TArray<TUniquePtr<FSocketEOS>> TrackedSockets;
 
-	/** All bound socket addresses */
+	/** All bound socket names */
 	TSet<FString> BoundSocketNames;
 
 	/** The last error we received */
