Minor fixes in FSocketEOS

diff --git a/Engine/Plugins/Online/SocketSubsystemEOS/Source/SocketSubsystemEOS/Private/SocketEOS.cpp b/Engine/Plugins/Online/SocketSubsystemEOS/Source/SocketSubsystemEOS/Private/SocketEOS.cpp
--- a/Engine/Plugins/Online/SocketSubsystemEOS/Source/SocketSubsystemEOS/Private/SocketEOS.cpp
+++ b/Engine/Plugins/Online/SocketSubsystemEOS/Source/SocketSubsystemEOS/Private/SocketEOS.cpp
@@ -32,7 +32,11 @@ FSocketEOS::FSocketEOS(FSocketSubsystemEOS& InSocketSubsystem, const FString& In
 			PacketReliability = PacketReliabilityType;
 		}
 	}
-	IEOSSDKManager::Get()->OnNetworkStatusChanged.AddRaw(this, &FSocketEOS::OnNetworkStatusChanged);
+
+	if (IEOSSDKManager* EOSSDKManager = IEOSSDKManager::Get())
+	{
+		NetworkStatusChangedDelegateHandle = EOSSDKManager->OnNetworkStatusChanged.AddRaw(this, &FSocketEOS::OnNetworkStatusChanged);
+	}
 #endif
 }
 
@@ -43,6 +47,13 @@ FSocketEOS::~FSocketEOS()
 	SocketSubsystem.UnbindSocketName(SocketName);
 	LocalAddress = FInternetAddrEOS();
 	CallbackAliveTracker = nullptr;
+
+#if WITH_EOS_SDK
+	if (IEOSSDKManager* EOSSDKManager = IEOSSDKManager::Get())
+	{
+		EOSSDKManager->OnNetworkStatusChanged.Remove(NetworkStatusChangedDelegateHandle);
+	}
+#endif
 }
 
 bool FSocketEOS::Shutdown(ESocketShutdownMode Mode)
@@ -267,6 +278,13 @@ bool FSocketEOS::SendTo(const uint8* Data, int32 Count, int32& OutBytesSent, con
 {
 	check(IsInGameThread() && "p2p does not support multithreading");
 
+	if (!SocketSubsystem.GetSocketSubsystemEOSUtils()->IsLoggedIn())
+	{
+		UE_LOG(LogSocketSubsystemEOS, VeryVerbose, TEXT("FSocketEOS::SendTo - Local user is logged out"));
+
+		return false;
+	}
+
 	OutBytesSent = 0;
 
 	if (!Destination.IsValid())
@@ -368,7 +386,7 @@ bool FSocketEOS::RecvFrom(uint8* Data, int32 BufferSize, int32& BytesRead, FInte
 
 	if (!SocketSubsystem.GetSocketSubsystemEOSUtils()->IsLoggedIn())
 	{
-		UE_LOG(LogSocketSubsystemEOS, VeryVerbose, TEXT("Local user is logged out"));
+		UE_LOG(LogSocketSubsystemEOS, VeryVerbose, TEXT("FSocketEOS::RecvFrom - Local user is logged out"));
 
 		return false;
 	}
diff --git a/Engine/Plugins/Online/SocketSubsystemEOS/Source/SocketSubsystemEOS/Public/SocketEOS.h b/Engine/Plugins/Online/SocketSubsystemEOS/Source/SocketSubsystemEOS/Public/SocketEOS.h
--- a/Engine/Plugins/Online/SocketSubsystemEOS/Source/SocketSubsystemEOS/Public/SocketEOS.h
+++ b/Engine/Plugins/Online/SocketSubsystemEOS/Source/SocketSubsystemEOS/Public/SocketEOS.h
@@ -106,5 +106,8 @@ private:
 
 	/** Default packet reliability type used in EOS_P2P_SendPacketOptions */
 	EOS_EPacketReliability PacketReliability = EOS_EPacketReliability::EOS_PR_UnreliableUnordered;
+
+	/** Delegate fired when the EOS SDK network status changes */
+	FDelegateHandle NetworkStatusChangedDelegateHandle;
 #endif
 };
