Removing port info from all EOS P2P connection strings.
#Synced-CL 49726841

diff --git a/Engine/Plugins/Online/OnlineSubsystemEOS/Source/OnlineSubsystemEOS/Private/OnlineSessionEOS.cpp b/Engine/Plugins/Online/OnlineSubsystemEOS/Source/OnlineSubsystemEOS/Private/OnlineSessionEOS.cpp
--- a/Engine/Plugins/Online/OnlineSubsystemEOS/Source/OnlineSubsystemEOS/Private/OnlineSessionEOS.cpp
+++ b/Engine/Plugins/Online/OnlineSubsystemEOS/Source/OnlineSubsystemEOS/Private/OnlineSessionEOS.cpp
@@ -3251,7 +3251,16 @@ static bool GetConnectStringFromSessionInfo(TSharedPtr<FOnlineSessionInfoEOS>& S
 		return false;
 	}
 
-	ConnectInfo = FString::Printf(TEXT("[%s]:%d"), *SessionInfo->HostAddr->ToString(false), PortOverride);
+	FString AddrString = SessionInfo->HostAddr->ToString(false);
+	if (AddrString.StartsWith(EOS_CONNECTION_URL_PREFIX, ESearchCase::IgnoreCase))
+	{
+		ConnectInfo = FString::Printf(TEXT("[%s]"), *AddrString);
+	}
+	else
+	{
+		int32 Port = PortOverride != 0 ? PortOverride : SessionInfo->HostAddr->GetPort();
+		ConnectInfo = FString::Printf(TEXT("[%s]:%d"), *SessionInfo->HostAddr->ToString(false), Port);
+	}
 
 	if (SocketNameOverride != NAME_None)
 	{
@@ -3961,7 +3970,12 @@ void FOnlineSessionEOS::SetPortFromNetDriver(const FOnlineSubsystemEOS& Subsyste
 	auto SessionInfoEOS = StaticCastSharedPtr<FOnlineSessionInfoEOS>(SessionInfo);
 	if (SessionInfoEOS.IsValid() && SessionInfoEOS->HostAddr.IsValid())
 	{
-		SessionInfoEOS->HostAddr->SetPort(NetDriverPort);
+		FString AddrString = SessionInfoEOS->HostAddr->ToString(false);
+		// Set port if IP address is used. EOS P2P does not use ports
+		if (!AddrString.StartsWith(EOS_CONNECTION_URL_PREFIX, ESearchCase::IgnoreCase))
+		{
+			SessionInfoEOS->HostAddr->SetPort(NetDriverPort);
+		}
 	}
 }
 
diff --git a/Engine/Plugins/Online/OnlineSubsystemUtils/Source/OnlineSubsystemUtils/Private/IpNetDriver.cpp b/Engine/Plugins/Online/OnlineSubsystemUtils/Source/OnlineSubsystemUtils/Private/IpNetDriver.cpp
--- a/Engine/Plugins/Online/OnlineSubsystemUtils/Source/OnlineSubsystemUtils/Private/IpNetDriver.cpp
+++ b/Engine/Plugins/Online/OnlineSubsystemUtils/Source/OnlineSubsystemUtils/Private/IpNetDriver.cpp
@@ -1372,7 +1372,10 @@ UNetConnection* UIpNetDriver::ProcessConnectionlessPacket(FReceivedPacketView& P
 						const FString OldAddress = RemoteAddrRef->ToString(true);
 
 						RemoteAddrRef->SetIp(*Address->ToString(false), bIsValid);
-						RemoteAddrRef->SetPort(Address->GetPort());
+						if (RemoteAddrRef->IsPortValid(Address->GetPort()))
+						{
+							RemoteAddrRef->SetPort(Address->GetPort());
+						}
 
 						MappedClientConnections.Add(RemoteAddrRef, FoundConn);
 
diff --git a/Engine/Plugins/Online/OnlineSubsystemUtils/Source/OnlineSubsystemUtils/Private/NetAddressResolution.cpp b/Engine/Plugins/Online/OnlineSubsystemUtils/Source/OnlineSubsystemUtils/Private/NetAddressResolution.cpp
--- a/Engine/Plugins/Online/OnlineSubsystemUtils/Source/OnlineSubsystemUtils/Private/NetAddressResolution.cpp
+++ b/Engine/Plugins/Online/OnlineSubsystemUtils/Source/OnlineSubsystemUtils/Private/NetAddressResolution.cpp
@@ -217,8 +217,10 @@ bool FNetConnectionAddressResolution::InitLocalConnection(ISocketSubsystem* Sock
 						TEXT("similar address, this will likely lead to issues!"));
 			}
 		}
-
-		RemoteAddr->SetPort(InURL.Port);
+		if (RemoteAddr->IsPortValid(InURL.Port))
+		{
+			RemoteAddr->SetPort(InURL.Port);
+		}
 
 		if (!bValidInit)
 		{
diff --git a/Engine/Plugins/Online/SocketSubsystemEOS/Source/SocketSubsystemEOS/Private/InternetAddrEOS.cpp b/Engine/Plugins/Online/SocketSubsystemEOS/Source/SocketSubsystemEOS/Private/InternetAddrEOS.cpp
--- a/Engine/Plugins/Online/SocketSubsystemEOS/Source/SocketSubsystemEOS/Private/InternetAddrEOS.cpp
+++ b/Engine/Plugins/Online/SocketSubsystemEOS/Source/SocketSubsystemEOS/Private/InternetAddrEOS.cpp
@@ -88,6 +88,12 @@ void FInternetAddrEOS::GetIp(uint32& OutAddr) const
 	UE_LOG(LogSocketSubsystemEOS, Warning, TEXT("Calls to FInternetAddrEOS::GetIp are not valid"));
 }
 
+bool FInternetAddrEOS::IsPortValid(int32 InPort) const
+{
+	// There are no ports in EOS connection strings
+	return false;
+}
+
 void FInternetAddrEOS::SetPort(int32 InPort)
 {
 	UE_LOG(LogSocketSubsystemEOS, Warning, TEXT("Calls to FInternetAddrEOS::SetPort are not valid"));
@@ -171,3 +177,9 @@ TSharedRef<FInternetAddr> FInternetAddrEOS::Clone() const
 {
 	return MakeShared<FInternetAddrEOS>(*this);
 }
+
+SOCKETSUBSYSTEMEOS_API FName FInternetAddrEOS::GetProtocolType() const
+{
+	static FName TypeName = FName(TEXT("EOS"));
+	return TypeName;
+}
diff --git a/Engine/Plugins/Online/SocketSubsystemEOS/Source/SocketSubsystemEOS/Private/NetConnectionEOS.cpp b/Engine/Plugins/Online/SocketSubsystemEOS/Source/SocketSubsystemEOS/Private/NetConnectionEOS.cpp
--- a/Engine/Plugins/Online/SocketSubsystemEOS/Source/SocketSubsystemEOS/Private/NetConnectionEOS.cpp
+++ b/Engine/Plugins/Online/SocketSubsystemEOS/Source/SocketSubsystemEOS/Private/NetConnectionEOS.cpp
@@ -24,12 +24,6 @@ void UNetConnectionEOS::InitLocalConnection(UNetDriver* InDriver, FSocket* InSoc
 	}
 
 	Super::InitLocalConnection(InDriver, InSocket, InURL, InState, InMaxPacket, InPacketOverhead);
-
-	// We re-set the original port, as if might have been overwritten by Super::InitLocalConnection
-	if (bHasP2PSession)
-	{
-		RemoteAddr->SetPort(InSocket->GetPortNo());
-	}
 }
 
 void UNetConnectionEOS::InitRemoteConnection(UNetDriver* InDriver, FSocket* InSocket, const FURL& InURL, const FInternetAddr& InRemoteAddr, EConnectionState InState, int32 InMaxPacket, int32 InPacketOverhead)
diff --git a/Engine/Plugins/Online/SocketSubsystemEOS/Source/SocketSubsystemEOS/Public/InternetAddrEOS.h b/Engine/Plugins/Online/SocketSubsystemEOS/Source/SocketSubsystemEOS/Public/InternetAddrEOS.h
--- a/Engine/Plugins/Online/SocketSubsystemEOS/Source/SocketSubsystemEOS/Public/InternetAddrEOS.h
+++ b/Engine/Plugins/Online/SocketSubsystemEOS/Source/SocketSubsystemEOS/Public/InternetAddrEOS.h
@@ -36,6 +36,7 @@ public:
 	SOCKETSUBSYSTEMEOS_API virtual void SetIp(uint32 InAddr) override;
 	SOCKETSUBSYSTEMEOS_API virtual void SetIp(const TCHAR* InAddr, bool& bIsValid) override;
 	SOCKETSUBSYSTEMEOS_API virtual void GetIp(uint32& OutAddr) const override;
+	SOCKETSUBSYSTEMEOS_API virtual bool IsPortValid(int32 InPort) const override;
 	SOCKETSUBSYSTEMEOS_API virtual void SetPort(int32 InPort) override;
 	SOCKETSUBSYSTEMEOS_API virtual int32 GetPort() const override;
 	SOCKETSUBSYSTEMEOS_API virtual void SetRawIp(const TArray<uint8>& RawAddr) override;
@@ -47,6 +48,7 @@ public:
 	SOCKETSUBSYSTEMEOS_API virtual uint32 GetTypeHash() const override;
 	SOCKETSUBSYSTEMEOS_API virtual bool IsValid() const override;
 	SOCKETSUBSYSTEMEOS_API virtual TSharedRef<FInternetAddr> Clone() const override;
+	SOCKETSUBSYSTEMEOS_API virtual FName GetProtocolType() const override;
 //~ End FInternetAddr Interface
 
 	inline FInternetAddrEOS& operator=(const FInternetAddrEOS& Other)
@@ -54,7 +56,17 @@ public:
 		ProductUserId = Other.ProductUserId;
 		return *this;
 	}
-	
+
+	inline virtual bool operator==(const FInternetAddr& Other) const override
+	{
+		if (Other.GetProtocolType() == GetProtocolType())
+		{
+			const FInternetAddrEOS& OtherEOS = static_cast<const FInternetAddrEOS&>(Other);
+			return *this == OtherEOS;
+		}
+		return false;
+	}
+
 	inline friend bool operator==(const FInternetAddrEOS& A, const FInternetAddrEOS& B)
 	{
 		return A.ProductUserId == B.ProductUserId;
