- Support EOS_P2P not being present in per-project SDK

#tests local tests launching Editor, did not crash, loads the DLL without the EOS_P2P symbols and does not complain, would presumably crash if we attempt to call one of the missing symbols but we avoid that with the bEnable config bool.

#Synced-CL 50524808

diff --git a/Engine/Plugins/Online/EOSShared/Source/EOSShared/Public/EOSShared.h b/Engine/Plugins/Online/EOSShared/Source/EOSShared/Public/EOSShared.h
--- a/Engine/Plugins/Online/EOSShared/Source/EOSShared/Public/EOSShared.h
+++ b/Engine/Plugins/Online/EOSShared/Source/EOSShared/Public/EOSShared.h
@@ -44,6 +44,10 @@
 #define WITH_EOS_RTC WITH_EOS_SDK && (EOS_MAJOR_VERSION >= 1 && EOS_MINOR_VERSION >= 13)
 #endif
 
+#ifndef WITH_EOS_P2P
+#define WITH_EOS_P2P WITH_EOS_SDK
+#endif
+
 #define EOS_ENUM_FORWARD_DECL(name) enum class name : int32_t;
 EOS_ENUM_FORWARD_DECL(EOS_EApplicationStatus);
 EOS_ENUM_FORWARD_DECL(EOS_EAuthScopeFlags);
diff --git a/Engine/Plugins/Online/OnlineServicesEOSGS/Source/Private/Online/SocketSubsystemEOSUtils_OnlineServicesEOSGS.cpp b/Engine/Plugins/Online/OnlineServicesEOSGS/Source/Private/Online/SocketSubsystemEOSUtils_OnlineServicesEOSGS.cpp
--- a/Engine/Plugins/Online/OnlineServicesEOSGS/Source/Private/Online/SocketSubsystemEOSUtils_OnlineServicesEOSGS.cpp
+++ b/Engine/Plugins/Online/OnlineServicesEOSGS/Source/Private/Online/SocketSubsystemEOSUtils_OnlineServicesEOSGS.cpp
@@ -22,6 +22,7 @@ FSocketSubsystemEOSUtils_OnlineServicesEOS::~FSocketSubsystemEOSUtils_OnlineServ
 {
 }
 
+#if WITH_EOS_P2P
 EOS_ProductUserId FSocketSubsystemEOSUtils_OnlineServicesEOS::GetLocalUserId()
 {
 	EOS_ProductUserId Result = nullptr;
@@ -45,6 +46,7 @@ EOS_ProductUserId FSocketSubsystemEOSUtils_OnlineServicesEOS::GetLocalUserId()
 
 	return Result;
 }
+#endif // WITH_EOS_P2P
 
 bool FSocketSubsystemEOSUtils_OnlineServicesEOS::IsLoggedIn()
 {
diff --git a/Engine/Plugins/Online/OnlineServicesEOSGS/Source/Private/Online/SocketSubsystemEOSUtils_OnlineServicesEOSGS.h b/Engine/Plugins/Online/OnlineServicesEOSGS/Source/Private/Online/SocketSubsystemEOSUtils_OnlineServicesEOSGS.h
--- a/Engine/Plugins/Online/OnlineServicesEOSGS/Source/Private/Online/SocketSubsystemEOSUtils_OnlineServicesEOSGS.h
+++ b/Engine/Plugins/Online/OnlineServicesEOSGS/Source/Private/Online/SocketSubsystemEOSUtils_OnlineServicesEOSGS.h
@@ -16,7 +16,9 @@ public:
 	FSocketSubsystemEOSUtils_OnlineServicesEOS(FOnlineServicesEOSGS& InServicesEOS);
 	virtual ~FSocketSubsystemEOSUtils_OnlineServicesEOS();
 
+#if WITH_EOS_P2P
 	virtual EOS_ProductUserId GetLocalUserId() override;
+#endif
 	virtual FString GetSessionId() override;
 	virtual FName GetSubsystemInstanceName() override;
 	virtual bool IsLoggedIn() override;
diff --git a/Engine/Plugins/Online/OnlineSubsystemEOS/Source/OnlineSubsystemEOS/Private/SocketSubsystemEOSUtils_OnlineSubsystemEOS.cpp b/Engine/Plugins/Online/OnlineSubsystemEOS/Source/OnlineSubsystemEOS/Private/SocketSubsystemEOSUtils_OnlineSubsystemEOS.cpp
--- a/Engine/Plugins/Online/OnlineSubsystemEOS/Source/OnlineSubsystemEOS/Private/SocketSubsystemEOSUtils_OnlineSubsystemEOS.cpp
+++ b/Engine/Plugins/Online/OnlineSubsystemEOS/Source/OnlineSubsystemEOS/Private/SocketSubsystemEOSUtils_OnlineSubsystemEOS.cpp
@@ -16,6 +16,7 @@ FSocketSubsystemEOSUtils_OnlineSubsystemEOS::~FSocketSubsystemEOSUtils_OnlineSub
 {
 }
 
+#if WITH_EOS_P2P
 EOS_ProductUserId FSocketSubsystemEOSUtils_OnlineSubsystemEOS::GetLocalUserId()
 {
 	EOS_ProductUserId Result = nullptr;
@@ -25,6 +26,7 @@ EOS_ProductUserId FSocketSubsystemEOSUtils_OnlineSubsystemEOS::GetLocalUserId()
 
 	return Result;
 }
+#endif // WITH_EOS_P2P
 
 bool FSocketSubsystemEOSUtils_OnlineSubsystemEOS::IsLoggedIn()
 {
diff --git a/Engine/Plugins/Online/OnlineSubsystemEOS/Source/OnlineSubsystemEOS/Private/SocketSubsystemEOSUtils_OnlineSubsystemEOS.h b/Engine/Plugins/Online/OnlineSubsystemEOS/Source/OnlineSubsystemEOS/Private/SocketSubsystemEOSUtils_OnlineSubsystemEOS.h
--- a/Engine/Plugins/Online/OnlineSubsystemEOS/Source/OnlineSubsystemEOS/Private/SocketSubsystemEOSUtils_OnlineSubsystemEOS.h
+++ b/Engine/Plugins/Online/OnlineSubsystemEOS/Source/OnlineSubsystemEOS/Private/SocketSubsystemEOSUtils_OnlineSubsystemEOS.h
@@ -12,7 +12,9 @@ public:
 	FSocketSubsystemEOSUtils_OnlineSubsystemEOS(FOnlineSubsystemEOS& InSubsystemEOS);
 	virtual ~FSocketSubsystemEOSUtils_OnlineSubsystemEOS() override;
 
+#if WITH_EOS_P2P
 	virtual EOS_ProductUserId GetLocalUserId() override;
+#endif // WITH_EOS_P2P
 	virtual FString GetSessionId() override;
 	virtual FName GetSubsystemInstanceName() override;
 	virtual bool IsLoggedIn() override;;
diff --git a/Engine/Plugins/Online/SocketSubsystemEOS/Source/SocketSubsystemEOS/Private/InternetAddrEOS.cpp b/Engine/Plugins/Online/SocketSubsystemEOS/Source/SocketSubsystemEOS/Private/InternetAddrEOS.cpp
--- a/Engine/Plugins/Online/SocketSubsystemEOS/Source/SocketSubsystemEOS/Private/InternetAddrEOS.cpp
+++ b/Engine/Plugins/Online/SocketSubsystemEOS/Source/SocketSubsystemEOS/Private/InternetAddrEOS.cpp
@@ -25,12 +25,12 @@ FInternetAddrEOS::FInternetAddrEOS(const FString& InProductUserId, const FString
 FInternetAddrEOS::FInternetAddrEOS(const FString& InProductUserId)
 	: ProductUserId(nullptr)
 {
-#if WITH_EOS_SDK
+#if WITH_EOS_P2P
 	ProductUserId = EOS_ProductUserId_FromString(TCHAR_TO_UTF8(*InProductUserId));
 #endif
 }
 
-#if WITH_EOS_SDK
+#if WITH_EOS_P2P
 FInternetAddrEOS::FInternetAddrEOS(const EOS_ProductUserId InProductUserId, const FString& InSocketName, const int32 InChannel)
 	: FInternetAddrEOS(InProductUserId)
 {
@@ -40,7 +40,11 @@ FInternetAddrEOS::FInternetAddrEOS(const EOS_ProductUserId InProductUserId, cons
 FInternetAddrEOS::FInternetAddrEOS(const EOS_ProductUserId InProductUserId)
 	: ProductUserId(InProductUserId)
 {
-
+}
+#else
+FInternetAddrEOS::FInternetAddrEOS(void* const InProductUserId)
+	: ProductUserId(InProductUserId)
+{
 }
 #endif
 
@@ -70,7 +74,7 @@ void FInternetAddrEOS::SetIp(const TCHAR* InAddr, bool& bIsValid)
 	{
 		return;
 	}
-#if WITH_EOS_SDK
+#if WITH_EOS_P2P
 	ProductUserId = EOS_ProductUserId_FromString(TCHAR_TO_UTF8(*UrlParts[1]));
 	if (EOS_ProductUserId_IsValid(ProductUserId) == EOS_FALSE)
 #endif
@@ -149,7 +153,7 @@ FString FInternetAddrEOS::ToString(bool bAppendPort) const
 {
 	char PuidBuffer[64];
 	int32 BufferLen = 64;
-#if WITH_EOS_SDK
+#if WITH_EOS_P2P
 	if (EOS_ProductUserId_ToString(ProductUserId, PuidBuffer, &BufferLen) != EOS_EResult::EOS_Success)
 #endif
 	{
@@ -166,7 +170,7 @@ uint32 FInternetAddrEOS::GetTypeHash() const
 
 bool FInternetAddrEOS::IsValid() const
 {
-#if WITH_EOS_SDK
+#if WITH_EOS_P2P
 	return (EOS_ProductUserId_IsValid(ProductUserId) == EOS_TRUE);
 #else
 	return false;
diff --git a/Engine/Plugins/Online/SocketSubsystemEOS/Source/SocketSubsystemEOS/Private/SocketEOS.cpp b/Engine/Plugins/Online/SocketSubsystemEOS/Source/SocketSubsystemEOS/Private/SocketEOS.cpp
--- a/Engine/Plugins/Online/SocketSubsystemEOS/Source/SocketSubsystemEOS/Private/SocketEOS.cpp
+++ b/Engine/Plugins/Online/SocketSubsystemEOS/Source/SocketSubsystemEOS/Private/SocketEOS.cpp
@@ -5,7 +5,7 @@
 #include "EOSShared.h"
 #include "Misc/ConfigCacheIni.h"
 
-#if WITH_EOS_SDK
+#if WITH_EOS_P2P
 	#include "eos_p2p.h"
 #endif
 
@@ -13,7 +13,7 @@ FSocketEOS::FSocketEOS(FSocketSubsystemEOS& InSocketSubsystem, const FString& In
 	: FSocket(ESocketType::SOCKTYPE_Datagram, InSocketDescription, NAME_None)
 	, SocketSubsystem(InSocketSubsystem)
 	, bIsListening(false)
-#if WITH_EOS_SDK
+#if WITH_EOS_P2P
 	, ConnectNotifyCallback(nullptr)
 	, ConnectNotifyId(EOS_INVALID_NOTIFICATIONID)
 	, ClosedNotifyCallback(nullptr)
@@ -22,7 +22,7 @@ FSocketEOS::FSocketEOS(FSocketSubsystemEOS& InSocketSubsystem, const FString& In
 {
 	CallbackAliveTracker = MakeShared<FCallbackBase>();
 
-#if WITH_EOS_SDK
+#if WITH_EOS_P2P
 	FString PacketReliabilityTypeStr;
 	if (GConfig->GetString(TEXT("SocketSubsystemEOS"), TEXT("PacketReliabilityType"), PacketReliabilityTypeStr, GEngineIni))
 	{
@@ -48,7 +48,7 @@ FSocketEOS::~FSocketEOS()
 	LocalAddress = FInternetAddrEOS();
 	CallbackAliveTracker = nullptr;
 
-#if WITH_EOS_SDK
+#if WITH_EOS_P2P
 	if (IEOSSDKManager* EOSSDKManager = IEOSSDKManager::Get())
 	{
 		EOSSDKManager->OnNetworkStatusChanged.Remove(NetworkStatusChangedDelegateHandle);
@@ -62,7 +62,7 @@ bool FSocketEOS::Shutdown(ESocketShutdownMode Mode)
 	return false;
 }
 
-#if WITH_EOS_SDK
+#if WITH_EOS_P2P
 void FSocketEOS::OnNetworkStatusChanged(EOS_ENetworkStatus OldNetworkStatus, EOS_ENetworkStatus NewNetworkStatus)
 {
 	// Only needed if we are listening on a socket, otherwise nothing needs to be done
@@ -90,7 +90,7 @@ bool FSocketEOS::Close()
 {
 	check(IsInGameThread() && "p2p does not support multithreading");
 
-#if WITH_EOS_SDK
+#if WITH_EOS_P2P
 	if (ConnectNotifyId != EOS_INVALID_NOTIFICATIONID)
 	{
 		EOS_P2P_RemoveNotifyPeerConnectionRequest(SocketSubsystem.GetP2PHandle(), ConnectNotifyId);
@@ -157,7 +157,7 @@ bool FSocketEOS::Listen(int32)
 {
 	check(IsInGameThread() && "p2p does not support multithreading");
 
-#if WITH_EOS_SDK
+#if WITH_EOS_P2P
 	// Add listener for inbound connections
 	EOS_P2P_SocketId SocketId = { };
 	SocketId.ApiVersion = 1;
@@ -232,7 +232,7 @@ bool FSocketEOS::HasPendingData(uint32& PendingDataSize)
 
 	PendingDataSize = 0;
 
-#if WITH_EOS_SDK
+#if WITH_EOS_P2P
 	EOS_P2P_GetNextReceivedPacketSizeOptions Options = { };
 	Options.ApiVersion = 2;
 	UE_EOS_CHECK_API_MISMATCH(EOS_P2P_GETNEXTRECEIVEDPACKETSIZE_API_LATEST, 2);
@@ -295,7 +295,7 @@ bool FSocketEOS::SendTo(const uint8* Data, int32 Count, int32& OutBytesSent, con
 		return false;
 	}
 
-#if WITH_EOS_SDK
+#if WITH_EOS_P2P
 	if (Count > EOS_P2P_MAX_PACKET_SIZE)
 	{
 		UE_LOG(LogSocketSubsystemEOS, Warning, TEXT("Unable to send data, data over maximum size. Amount=[%d/%d] DestinationAddress=[%s]"), Count, EOS_P2P_MAX_PACKET_SIZE, *Destination.ToString(true));
@@ -332,7 +332,7 @@ bool FSocketEOS::SendTo(const uint8* Data, int32 Count, int32& OutBytesSent, con
 		return false;
 	}
 
-#if WITH_EOS_SDK
+#if WITH_EOS_P2P
 	// Need to handle closures if we are a client and the server closes down on us
 	RegisterClosedNotification();
 
@@ -408,7 +408,7 @@ bool FSocketEOS::RecvFrom(uint8* Data, int32 BufferSize, int32& BytesRead, FInte
 		return false;
 	}
 
-#if WITH_EOS_SDK
+#if WITH_EOS_P2P
 	EOS_P2P_ReceivePacketOptions Options = { };
 	Options.ApiVersion = 2;
 	UE_EOS_CHECK_API_MISMATCH(EOS_P2P_RECEIVEPACKET_API_LATEST, 2);
@@ -583,7 +583,7 @@ bool FSocketEOS::Close(const FInternetAddrEOS& RemoteAddress)
 		return false;
 	}
 
-#if WITH_EOS_SDK
+#if WITH_EOS_P2P
 	// So we don't reopen a connection by sending to it
 	ClosedRemotes.Add(RemoteAddress);
 
@@ -619,7 +619,7 @@ bool FSocketEOS::Close(const FInternetAddrEOS& RemoteAddress)
 
 void FSocketEOS::RegisterClosedNotification()
 {
-#if WITH_EOS_SDK
+#if WITH_EOS_P2P
 	if (ClosedNotifyId != EOS_INVALID_NOTIFICATIONID)
 	{
 		// Already listening for these events so ignore
diff --git a/Engine/Plugins/Online/SocketSubsystemEOS/Source/SocketSubsystemEOS/Private/SocketSubsystemEOS.cpp b/Engine/Plugins/Online/SocketSubsystemEOS/Source/SocketSubsystemEOS/Private/SocketSubsystemEOS.cpp
--- a/Engine/Plugins/Online/SocketSubsystemEOS/Source/SocketSubsystemEOS/Private/SocketSubsystemEOS.cpp
+++ b/Engine/Plugins/Online/SocketSubsystemEOS/Source/SocketSubsystemEOS/Private/SocketSubsystemEOS.cpp
@@ -1,16 +1,16 @@
 // Copyright Epic Games, Inc. All Rights Reserved.
 
 #include "SocketSubsystemEOS.h"
+
+#include "EOSShared.h"
+#include "Misc/ConfigCacheIni.h"
+#include "OnlineSubsystemUtils.h"
 #include "SocketEOS.h"
 #include "SocketSubsystemModule.h"
-#include "OnlineSubsystemUtils.h"
-#include "EOSShared.h"
 
-#if WITH_EOS_SDK
-	#include "eos_sdk.h"
-	#include "eos_p2p.h"
-	#include "Misc/ConfigCacheIni.h"
-#endif
+#if WITH_EOS_P2P
+#include "eos_sdk.h"
+#include "eos_p2p.h"
 
 namespace
 {
@@ -35,22 +35,28 @@ namespace
 	}
 
 }
+#endif // WITH_EOS_P2P
 
 TArray<FSocketSubsystemEOS*> FSocketSubsystemEOS::SocketSubsystemEOSInstances;
 TMap<UWorld*, FSocketSubsystemEOS*> FSocketSubsystemEOS::SocketSubsystemEOSPerWorldMap;
 
 FSocketSubsystemEOS::FSocketSubsystemEOS(IEOSPlatformHandlePtr InPlatformHandle, ISocketSubsystemEOSUtilsPtr InUtils)
-	: P2PHandle(nullptr)
-	, Utils(InUtils)
+	: Utils(InUtils)
 	, LastSocketError(ESocketErrors::SE_NO_ERROR)
 {
-#if WITH_EOS_SDK
-	P2PHandle = EOS_Platform_GetP2PInterface(*InPlatformHandle);
-	if (P2PHandle == nullptr)
+#if WITH_EOS_P2P
+	bool bEnable = true;
+	GConfig->GetBool(TEXT("SocketSubsystemEOS"), TEXT("bEnable"), bEnable, GEngineIni);
+
+	if (bEnable)
 	{
-		UE_LOG(LogSocketSubsystemEOS, Error, TEXT("FOnlineSubsystemEOS: failed to init EOS platform, couldn't get p2p handle"));
+		P2PHandle = EOS_Platform_GetP2PInterface(*InPlatformHandle);
+		if (P2PHandle == nullptr)
+		{
+			UE_LOG(LogSocketSubsystemEOS, Error, TEXT("FOnlineSubsystemEOS: failed to init EOS platform, couldn't get p2p handle"));
+		}
 	}
-#endif
+#endif // WITH_EOS_P2P
 }
 
 FSocketSubsystemEOS::~FSocketSubsystemEOS()
@@ -86,24 +92,27 @@ bool FSocketSubsystemEOS::Init(FString& Error)
 	FSocketSubsystemModule& SocketSubsystem = FModuleManager::LoadModuleChecked<FSocketSubsystemModule>("Sockets");
 	SocketSubsystem.RegisterSocketSubsystem(EOS_SOCKETSUBSYSTEM, this, false);
 
-#if WITH_EOS_SDK
-	FString RelayControlStr;
-	GConfig->GetString(TEXT("SocketSubsystemEOS"), TEXT("RelayControl"), RelayControlStr, GEngineIni);
-	if (!RelayControlStr.IsEmpty())
+	#if WITH_EOS_P2P
+	if (P2PHandle)
 	{
-		EOS_P2P_SetRelayControlOptions Options = {};
-		Options.ApiVersion = 1;
-		UE_EOS_CHECK_API_MISMATCH(EOS_P2P_SETRELAYCONTROL_API_LATEST, 1);
+		FString RelayControlStr;
+		GConfig->GetString(TEXT("SocketSubsystemEOS"), TEXT("RelayControl"), RelayControlStr, GEngineIni);
+		if (!RelayControlStr.IsEmpty())
+		{
+			EOS_P2P_SetRelayControlOptions Options = {};
+			Options.ApiVersion = 1;
+			UE_EOS_CHECK_API_MISMATCH(EOS_P2P_SETRELAYCONTROL_API_LATEST, 1);
 
-		Options.RelayControl = LexFromString(RelayControlStr);
+			Options.RelayControl = LexFromString(RelayControlStr);
 
-		EOS_EResult Result = EOS_P2P_SetRelayControl(P2PHandle, &Options);
-		if (Result != EOS_EResult::EOS_Success && Result != EOS_EResult::EOS_NoChange)
-		{
-			UE_LOG(LogSocketSubsystemEOS, Warning, TEXT("[FSocketSubsystemEOS::Init] EOS_P2P_SetRelayControl failed with result [%s]"), *LexToString(Result));
+			EOS_EResult Result = EOS_P2P_SetRelayControl(P2PHandle, &Options);
+			if (Result != EOS_EResult::EOS_Success && Result != EOS_EResult::EOS_NoChange)
+			{
+				UE_LOG(LogSocketSubsystemEOS, Warning, TEXT("[FSocketSubsystemEOS::Init] EOS_P2P_SetRelayControl failed with result [%s]"), *LexToString(Result));
+			}
 		}
 	}
-#endif
+#endif // WITH_EOS_P2P
 
 	return true;
 }
@@ -228,7 +237,7 @@ TSharedRef<FInternetAddr> FSocketSubsystemEOS::GetLocalBindAddr(FOutputDevice& O
 	return GetLocalBindAddr(nullptr, Out);
 }
 
-#if WITH_EOS_SDK
+#if WITH_EOS_P2P
 EOS_HP2P FSocketSubsystemEOS::GetP2PHandle()
 {
 	check(P2PHandle != nullptr);
@@ -239,13 +248,13 @@ EOS_ProductUserId FSocketSubsystemEOS::GetLocalUserId()
 {
 	return Utils->GetLocalUserId();
 }
-#endif
+#endif // WITH_EOS_P2P
 
 TSharedRef<FInternetAddr> FSocketSubsystemEOS::GetLocalBindAddr(const UWorld* const OwningWorld, FOutputDevice& Out)
 {
 	TSharedRef<FInternetAddrEOS> BoundAddr = MakeShared<FInternetAddrEOS>();
 
-#if WITH_EOS_SDK
+#if WITH_EOS_P2P
 	EOS_ProductUserId LocalUserId = GetLocalUserId();
 	if (LocalUserId != nullptr)
 	{
@@ -255,7 +264,7 @@ TSharedRef<FInternetAddr> FSocketSubsystemEOS::GetLocalBindAddr(const UWorld* co
 	{
 		UE_LOG(LogSocketSubsystemEOS, Error, TEXT("No local user set in LocalBindAddr"));
 	}
-#endif
+#endif // WITH_EOS_P2P
 
 	return BoundAddr;
 }
diff --git a/Engine/Plugins/Online/SocketSubsystemEOS/Source/SocketSubsystemEOS/Public/InternetAddrEOS.h b/Engine/Plugins/Online/SocketSubsystemEOS/Source/SocketSubsystemEOS/Public/InternetAddrEOS.h
--- a/Engine/Plugins/Online/SocketSubsystemEOS/Source/SocketSubsystemEOS/Public/InternetAddrEOS.h
+++ b/Engine/Plugins/Online/SocketSubsystemEOS/Source/SocketSubsystemEOS/Public/InternetAddrEOS.h
@@ -3,9 +3,10 @@
 #pragma once
 
 #include "CoreMinimal.h"
+#include "EOSShared.h"
 #include "IPAddress.h"
 
-#if WITH_EOS_SDK
+#if WITH_EOS_P2P
 #if defined(EOS_PLATFORM_BASE_FILE_NAME)
 #include EOS_PLATFORM_BASE_FILE_NAME
 #endif
@@ -25,10 +26,12 @@ public:
 	UE_DEPRECATED(5.6, "InSocketName and InChannel have been removed from FInternetAddrEOS, please use the new constructor.")
 	SOCKETSUBSYSTEMEOS_API FInternetAddrEOS(const FString& InRemoteUserId, const FString& InSocketName, const int32 InChannel);
 	SOCKETSUBSYSTEMEOS_API FInternetAddrEOS(const FString& InRemoteUserId);
-#if WITH_EOS_SDK
+#if WITH_EOS_P2P
 	UE_DEPRECATED(5.6, "InSocketName and InChannel have been removed from FInternetAddrEOS, please use the new constructor.")
 	SOCKETSUBSYSTEMEOS_API FInternetAddrEOS(const EOS_ProductUserId InRemoteUserId, const FString& InSocketName, const int32 InChannel);
 	SOCKETSUBSYSTEMEOS_API FInternetAddrEOS(const EOS_ProductUserId InRemoteUserId);
+#else
+	SOCKETSUBSYSTEMEOS_API FInternetAddrEOS(void* const InRemoteUserId);
 #endif
 	virtual ~FInternetAddrEOS() = default;
 
@@ -87,7 +90,7 @@ public:
 		return Left.ProductUserId < Right.ProductUserId;
 	}
 
-#if WITH_EOS_SDK
+#if WITH_EOS_P2P
 	void SetProductUserId(EOS_ProductUserId InProductUserId)
 	{
 		ProductUserId = InProductUserId;
@@ -110,7 +113,7 @@ public:
 #endif
 
 private:
-#if WITH_EOS_SDK
+#if WITH_EOS_P2P
 	EOS_ProductUserId ProductUserId;
 #else
 	void* ProductUserId;
diff --git a/Engine/Plugins/Online/SocketSubsystemEOS/Source/SocketSubsystemEOS/Public/SocketEOS.h b/Engine/Plugins/Online/SocketSubsystemEOS/Source/SocketSubsystemEOS/Public/SocketEOS.h
--- a/Engine/Plugins/Online/SocketSubsystemEOS/Source/SocketSubsystemEOS/Public/SocketEOS.h
+++ b/Engine/Plugins/Online/SocketSubsystemEOS/Source/SocketSubsystemEOS/Public/SocketEOS.h
@@ -10,7 +10,7 @@
 class FOnlineSubsystemEOS;
 class FSocketSubsystemEOS;
 
-#if WITH_EOS_SDK
+#if WITH_EOS_P2P
 	#include "eos_p2p_types.h"
 	#include "eos_types.h"
 #endif
@@ -94,7 +94,7 @@ private:
 
 	TArray<FInternetAddrEOS> ClosedRemotes;
 
-#if WITH_EOS_SDK
+#if WITH_EOS_P2P
 	typedef TEOSGlobalCallback<EOS_P2P_OnIncomingConnectionRequestCallback, EOS_P2P_OnIncomingConnectionRequestInfo, FCallbackBase> FConnectNotifyCallback;
 	typedef TEOSGlobalCallback<EOS_P2P_OnRemoteConnectionClosedCallback, EOS_P2P_OnRemoteConnectionClosedInfo, FCallbackBase> FClosedNotifyCallback;
 
diff --git a/Engine/Plugins/Online/SocketSubsystemEOS/Source/SocketSubsystemEOS/Public/SocketSubsystemEOS.h b/Engine/Plugins/Online/SocketSubsystemEOS/Source/SocketSubsystemEOS/Public/SocketSubsystemEOS.h
--- a/Engine/Plugins/Online/SocketSubsystemEOS/Source/SocketSubsystemEOS/Public/SocketSubsystemEOS.h
+++ b/Engine/Plugins/Online/SocketSubsystemEOS/Source/SocketSubsystemEOS/Public/SocketSubsystemEOS.h
@@ -3,16 +3,17 @@
 #pragma once
 
 #include "CoreMinimal.h"
-#include "SocketSubsystem.h"
+#include "EOSShared.h"
 #include "IEOSSDKManager.h"
+#include "SocketSubsystem.h"
 #include "SocketSubsystemEOSUtils.h"
 
-#if WITH_EOS_SDK
-	#if defined(EOS_PLATFORM_BASE_FILE_NAME)
-	#include EOS_PLATFORM_BASE_FILE_NAME
-	#endif
+#if WITH_EOS_P2P
+#if defined(EOS_PLATFORM_BASE_FILE_NAME)
+#include EOS_PLATFORM_BASE_FILE_NAME
+#endif
 
-	#include "eos_p2p_types.h"
+#include "eos_p2p_types.h"
 #endif
 
 #ifndef EOS_SOCKETSUBSYSTEM
@@ -61,7 +62,7 @@ public:
 
 	void SetLastSocketError(const ESocketErrors NewSocketError);
 
-#if WITH_EOS_SDK
+#if WITH_EOS_P2P
 	EOS_HP2P GetP2PHandle();
 	EOS_ProductUserId GetLocalUserId();
 #endif
@@ -101,8 +102,8 @@ private:
 	void RemoveFromStaticContainers();
 
 private:
-#if WITH_EOS_SDK
-	EOS_HP2P P2PHandle;
+#if WITH_EOS_P2P
+	EOS_HP2P P2PHandle = nullptr;
 #endif
 	ISocketSubsystemEOSUtilsPtr Utils;
 
diff --git a/Engine/Plugins/Online/SocketSubsystemEOS/Source/SocketSubsystemEOS/Public/SocketSubsystemEOSUtils.h b/Engine/Plugins/Online/SocketSubsystemEOS/Source/SocketSubsystemEOS/Public/SocketSubsystemEOSUtils.h
--- a/Engine/Plugins/Online/SocketSubsystemEOS/Source/SocketSubsystemEOS/Public/SocketSubsystemEOSUtils.h
+++ b/Engine/Plugins/Online/SocketSubsystemEOS/Source/SocketSubsystemEOS/Public/SocketSubsystemEOSUtils.h
@@ -3,8 +3,9 @@
 #pragma once
 
 #include "CoreMinimal.h"
+#include "EOSShared.h"
 
-#if WITH_EOS_SDK
+#if WITH_EOS_P2P
 #if defined(EOS_PLATFORM_BASE_FILE_NAME)
 #include EOS_PLATFORM_BASE_FILE_NAME
 #endif
@@ -17,7 +18,7 @@ class ISocketSubsystemEOSUtils
 public:
 	virtual ~ISocketSubsystemEOSUtils() = default;
 
-#if WITH_EOS_SDK
+#if WITH_EOS_P2P
 	virtual EOS_ProductUserId GetLocalUserId() = 0;
 #endif
 	virtual FString GetSessionId() = 0;
